Move GDK modules into public engine

#rnx

#Synced-CL 50642529

#Synced-CL 50642529

diff --git a/Engine/Source/Developer/Microsoft/GDK/GDKPlatformEditor/GDKPlatformEditor.Build.cs b/Engine/Source/Developer/Microsoft/GDK/GDKPlatformEditor/GDKPlatformEditor.Build.cs
new file mode 100644
--- /dev/null
+++ b/Engine/Source/Developer/Microsoft/GDK/GDKPlatformEditor/GDKPlatformEditor.Build.cs
@@ -0,0 +1,43 @@
+// Copyright Epic Games, Inc. All Rights Reserved.
+
+using UnrealBuildTool;
+
+public class GDKPlatformEditor : ModuleRules
+{
+	public GDKPlatformEditor(ReadOnlyTargetRules Target) : base(Target)
+	{
+		BinariesSubFolder = "GDK";
+
+		PrivateDependencyModuleNames.AddRange(
+			new string[] {
+				"Core",
+				"CoreUObject",
+				"InputCore",
+				"Engine",
+				"MainFrame",
+				"Slate",
+				"SlateCore",
+				"PropertyEditor",
+				"SharedSettingsWidgets",
+				"SourceControl",
+				"UnrealEd",
+				"DeveloperToolSettings",
+				"XmlParser",
+				"Projects",
+				"ToolWidgets",
+            }
+		);
+
+		PrivateIncludePathModuleNames.AddRange(
+			new string[] {
+				"Settings",
+			}
+		);
+
+		PublicIncludePathModuleNames.AddRange(
+			new string[] {
+				"PropertyEditor",
+			}
+		);
+	}
+}
diff --git a/Engine/Source/Developer/Microsoft/GDK/GDKPlatformEditor/Private/GDKPlatformEditorModule.cpp b/Engine/Source/Developer/Microsoft/GDK/GDKPlatformEditor/Private/GDKPlatformEditorModule.cpp
new file mode 100644
--- /dev/null
+++ b/Engine/Source/Developer/Microsoft/GDK/GDKPlatformEditor/Private/GDKPlatformEditorModule.cpp
@@ -0,0 +1,181 @@
+// Copyright Epic Games, Inc. All Rights Reserved.
+
+#include "GDKPlatformEditorModule.h"
+#include "Modules/ModuleManager.h"
+#include "Modules/ModuleInterface.h"
+#include "ISettingsModule.h"
+#include "Modules/ModuleManager.h"
+#include "UObject/Object.h"
+#include "GDKTargetSettings.h"
+#include "UObject/WeakObjectPtr.h"
+#include "UObject/Class.h"
+#include "PropertyEditorModule.h"
+#include "GDKTargetSettingsDetails.h"
+#include "XmlParser.h"
+#include "Misc/MonitoredProcess.h"
+#include "Dialog/SMessageDialog.h"
+#include "Async/Async.h"
+#include "HAL/CriticalSection.h"
+#include "Misc/ScopeLock.h"
+
+
+#define LOCTEXT_NAMESPACE "FGDKPlatformEditorModule"
+
+DEFINE_LOG_CATEGORY_STATIC(LogGDKPlatform, Log, All);
+
+
+/**
+* Module for GDK platform editor utilities
+*/
+class FGDKPlatformEditorModule
+	: public IGDKPlatformEditorModule
+{
+	virtual void StartupModule() override
+	{
+		// register detail customization
+		static FName PropertyEditor("PropertyEditor");
+		FPropertyEditorModule* PropertyModule = FModuleManager::GetModulePtr<FPropertyEditorModule>(PropertyEditor);
+
+		if (PropertyModule != nullptr)
+		{
+			// customizing the string resources struct so that the image picker dialog button can be included too
+			static FName GDKPackageStringResources("GDKPackageStringResources");
+			PropertyModule->RegisterCustomPropertyTypeLayout(GDKPackageStringResources, FOnGetPropertyTypeCustomizationInstance::CreateStatic(&FGDKCultureResourceDetails::MakeInstance));
+		}
+	}
+
+	virtual void ShutdownModule() override
+	{
+	}
+
+
+	void PartnerCenterError( const FText& Message )
+	{
+		TSharedRef<SMessageDialog> Dialog = SNew(SMessageDialog)
+			.Title(LOCTEXT("PartnerCenterErrorTitle","Partner Center"))
+			.Message(Message)
+			.Buttons({SMessageDialog::FButton(LOCTEXT("OK","OK"))})
+		;
+		Dialog->ShowModal();
+	}
+
+	virtual void PartnerCenterQueryProductsAsync( TFunction<void(bool)> OnComplete )
+	{
+		FScopeLock Lock(&PartnerCenterToolProcessCS);
+
+		// make sure we're not already running a query
+		if (PartnerCenterToolProcess.IsValid())
+		{
+			OnComplete(false);
+			return;
+		}
+
+		// make sure the tool exists
+		FString PartnerCenterTool = FPaths::Combine( FPaths::EngineDir(), TEXT("Binaries\\DotNET\\Microsoft\\GDK\\PartnerCenterTool\\PartnerCenterTool.exe"));
+		if (!FPaths::FileExists(PartnerCenterTool))
+		{
+			PartnerCenterError( FText::Format( LOCTEXT("PartnerCenterToolNotFoundFmt","{0} - file not found"), FText::FromString(PartnerCenterTool) ));
+			OnComplete(false);
+			return;
+		}
+
+		// run the tool
+		PartnerCenterToolProcess = MakeShared<FMonitoredProcess>(PartnerCenterTool, TEXT("products"), true );
+		PartnerCenterToolProcess->OnCompleted().BindLambda( [this, OnComplete = MoveTemp(OnComplete)](int32) mutable
+		{
+			// defer one tick so we can read the process output
+			AsyncTask( ENamedThreads::GameThread, [this, OnComplete = MoveTemp(OnComplete)]() mutable
+			{
+				bool bResult = ParsePartnerCenterQueryResults();
+				OnComplete(bResult);
+				PartnerCenterToolProcess.Reset();
+			});
+		});
+		if (!PartnerCenterToolProcess->Launch())
+		{
+			PartnerCenterError( FText::Format( LOCTEXT("PartnerCenterToolLaunchFail","{0} - launch failed"), FText::FromString(PartnerCenterTool) ));
+			OnComplete(false);
+			PartnerCenterToolProcess.Reset();
+			return;
+		}
+	}
+
+	bool ParsePartnerCenterQueryResults()
+	{
+		FScopeLock Lock(&PartnerCenterToolProcessCS);
+
+		if (!PartnerCenterToolProcess.IsValid())
+		{
+			return false;
+		}
+
+		const FString& FullOutput = PartnerCenterToolProcess->GetFullOutputWithoutDelegate();
+		if (!FullOutput.IsEmpty())
+		{
+			// read the xml output & recreate product list
+			FXmlFile XmlFile( FullOutput, EConstructMethod::ConstructFromBuffer );
+			if (XmlFile.IsValid() && XmlFile.GetRootNode() != nullptr)
+			{
+				PartnerCenterProducts.Reset();
+				for (FXmlNode* ProductNode : XmlFile.GetRootNode()->GetChildrenNodes() )
+				{
+					TSharedPtr<IGDKPlatformEditorModule::FPartnerCenterProduct> Product = MakeShared<IGDKPlatformEditorModule::FPartnerCenterProduct>();
+					PartnerCenterProducts.Add(Product);
+
+					// ensure these two values are always defined to avoid checking for them everywhere they're used
+					Product->Add(TEXT("PackageDisplayName"), TEXT("(unnamed product)"));
+					Product->Add(TEXT("PackageName"),        TEXT("(invalid package name)"));
+
+					// populate the product dictionary
+					for (FXmlNode* DataNode : ProductNode->GetChildrenNodes())
+					{
+						Product->Add(DataNode->GetTag(), DataNode->GetContent());
+					}
+				}
+
+				// sort the product list based on display name
+				PartnerCenterProducts.Sort( []( const TSharedPtr<FPartnerCenterProduct> A, const TSharedPtr<FPartnerCenterProduct> B )
+				{
+					const TCHAR* SortKey = TEXT("PackageDisplayName");
+					return A->FindRef(SortKey) < B->FindRef(SortKey);
+				});
+
+				return true;
+			}
+		}
+
+		PartnerCenterError( FText::Format( LOCTEXT("PartnerCenterToolFailFmt", "Could not retrieve titles from Partner Center\n{0}"), FText::FromString(FullOutput) ));
+		return false;
+	}
+
+	virtual const TArray<TSharedPtr<FPartnerCenterProduct>>* GetPartnerCenterProducts() override
+	{
+		return &PartnerCenterProducts;
+	}
+
+	virtual bool IsQueryingPartnerCenter() const override
+	{
+		FScopeLock Lock(&PartnerCenterToolProcessCS);
+		return PartnerCenterToolProcess.IsValid();
+	}
+
+	virtual void CancelQueryPartnerCenter() override
+	{
+		FScopeLock Lock(&PartnerCenterToolProcessCS);
+
+		if (PartnerCenterToolProcess.IsValid())
+		{
+			PartnerCenterToolProcess->Cancel();
+			PartnerCenterToolProcess.Reset();
+		}
+	}
+
+	TSharedPtr<FMonitoredProcess> PartnerCenterToolProcess;
+	TArray<TSharedPtr<FPartnerCenterProduct>> PartnerCenterProducts;
+	mutable FCriticalSection PartnerCenterToolProcessCS;
+};
+
+
+IMPLEMENT_MODULE(FGDKPlatformEditorModule, GDKPlatformEditor);
+
+#undef LOCTEXT_NAMESPACE
diff --git a/Engine/Source/Developer/Microsoft/GDK/GDKPlatformEditor/Private/GDKTargetSettings.cpp b/Engine/Source/Developer/Microsoft/GDK/GDKPlatformEditor/Private/GDKTargetSettings.cpp
new file mode 100644
--- /dev/null
+++ b/Engine/Source/Developer/Microsoft/GDK/GDKPlatformEditor/Private/GDKTargetSettings.cpp
@@ -0,0 +1,56 @@
+// Copyright Epic Games, Inc. All Rights Reserved.
+
+#include "GDKTargetSettings.h"
+#include "Settings/ProjectPackagingSettings.h"
+#include "Interfaces/IPluginManager.h"
+
+#include UE_INLINE_GENERATED_CPP_BY_NAME(GDKTargetSettings)
+
+UGDKTargetSettings::UGDKTargetSettings(const FObjectInitializer& ObjectInitializer)
+	: Super(ObjectInitializer)
+{
+}
+
+
+TArray<FString> UGDKTargetSettings::GetKnownDeviceTypes()
+{
+	return TArray<FString>();
+}
+
+TArray<FString> UGDKTargetSettings::GetKnownIntelligentDeliveryFeatures()
+{
+	TArray<FString> KnownFeatures;
+	for ( const FGDKIntelligentDeliveryFeature& Feature : IntelligentDeliveryFeatures )
+	{
+		KnownFeatures.AddUnique(Feature.Id);
+	}
+
+	KnownFeatures.Sort();
+	return KnownFeatures;	
+}
+
+TArray<FString> UGDKTargetSettings::GetKnownCultureStageIds()
+{
+	return GetDefault<UProjectPackagingSettings>()->CulturesToStage;	
+}
+
+TArray<FString> UGDKTargetSettings::GetKnownDLCNames()
+{
+	TArray<FString> KnownDLC;
+
+	for (TSharedRef<IPlugin> Plugin : IPluginManager::Get().GetEnabledPluginsWithContent())
+	{
+		if (Plugin->GetType() == EPluginType::Project && Plugin->GetLoadedFrom() == EPluginLoadedFrom::Project && Plugin->GetDescriptor().SupportsTargetPlatform(GetPlatformName()) )
+		{
+			FString DLCName = Plugin->GetName();
+			KnownDLC.Add(DLCName);
+		}
+	}
+
+	return KnownDLC;
+}
+
+bool UGDKTargetSettings::CanEditProductId() const
+{
+	return StoreId.IsEmpty();
+}
diff --git a/Engine/Source/Developer/Microsoft/GDK/GDKPlatformEditor/Private/GDKTargetSettingsDetails.cpp b/Engine/Source/Developer/Microsoft/GDK/GDKPlatformEditor/Private/GDKTargetSettingsDetails.cpp
new file mode 100644
--- /dev/null
+++ b/Engine/Source/Developer/Microsoft/GDK/GDKPlatformEditor/Private/GDKTargetSettingsDetails.cpp
@@ -0,0 +1,878 @@
+// Copyright Epic Games, Inc. All Rights Reserved.
+
+#include "GDKTargetSettingsDetails.h"
+
+#include "SExternalImageReference.h"
+#include "EditorDirectories.h"
+#include "DetailWidgetRow.h"
+#include "DetailCategoryBuilder.h"
+#include "DetailLayoutBuilder.h"
+#include "Framework/Application/SlateApplication.h"
+#include "IDetailChildrenBuilder.h"
+#include "IDetailGroup.h"
+#include "UObject/UnrealType.h"
+#include "Widgets/Text/STextBlock.h"
+#include "Widgets/Input/SHyperlink.h"
+#include "Widgets/Layout/SBox.h"
+#include "Widgets/Layout/SScrollBox.h"
+#include "Widgets/Layout/SWidgetSwitcher.h"
+#include "Widgets/Images/SImage.h"
+#include "Widgets/Images/SThrobber.h"
+#include "Widgets/Views/SListView.h"
+#include "Misc/MessageDialog.h"
+#include "Misc/NotifyHook.h"
+#include "Misc/DataDrivenPlatformInfoRegistry.h"
+#include "HAL/PlatformFileManager.h"
+#include "Dialogs/Dialogs.h"
+#include "GDKTargetSettings.h"
+#include "ObjectEditorUtils.h"
+#include "ISourceControlModule.h"
+#include "ISourceControlProvider.h"
+#include "SourceControlOperations.h"
+#include "IPropertyUtilities.h"
+#include "Interfaces/IMainFrameModule.h"
+#include "Interfaces/IPluginManager.h"
+#include "Styling/AppStyle.h"
+
+#define LOCTEXT_NAMESPACE "GDKTargetSettingsDetails"
+
+namespace GDKTargetSettingsDetailsConstants
+{
+	/** Relative path from engine platform extensions or project platform extensions dir to build resources */
+	const FString BuildResourcePath(TEXT("Build/Resources"));
+	const FString EngineResourcePath(TEXT("Build/DefaultImages"));
+}
+
+
+
+FGDKTargetSettingsDetails::FGDKTargetSettingsDetails(const TCHAR* InPlatformName)
+	: PlatformName(InPlatformName)
+{
+}
+
+
+void FGDKTargetSettingsDetails::CustomizeDetails(IDetailLayoutBuilder& DetailBuilder)
+{
+	IDetailCategoryBuilder& PartnerCenterCategory = DetailBuilder.EditCategory("PartnerCenter", FText::GetEmpty(), ECategoryPriority::Important);
+	IDetailCategoryBuilder& RecipesCategory = DetailBuilder.EditCategory("Chunk Install", FText::GetEmpty(), ECategoryPriority::Uncommon);
+	const TSharedRef<IPropertyUtilities> PropertyUtilities = DetailBuilder.GetPropertyUtilities();
+
+	// find the target settings that are being customized
+	TArray<TWeakObjectPtr<UObject>> ObjectsBeingCustomized;
+	DetailBuilder.GetObjectsBeingCustomized(ObjectsBeingCustomized);
+
+	TWeakObjectPtr<UGDKTargetSettings> GDKTargetSettings = nullptr;
+	for (auto Itr : ObjectsBeingCustomized)
+	{
+		GDKTargetSettings = Cast<UGDKTargetSettings>(Itr.Get());
+		if (GDKTargetSettings.IsValid())
+		{
+			break;
+		}
+	}
+
+
+	// add the partner center toolbar
+	TSharedPtr<SHorizontalBox> PartnerCenterBox;
+	PartnerCenterCategory.AddCustomRow(FText::GetEmpty(), false)
+	.WholeRowWidget
+	[
+		SNew(SBorder)
+		.BorderImage(FAppStyle::GetBrush("ToolPanel.GroupBorder"))
+		[
+			SAssignNew(PartnerCenterBox, SHorizontalBox)
+		]
+	];
+
+	// if there's nothing set up for this project yet, give the option to configure the game from partner center etc.
+	AddGettingStartedItems(GDKTargetSettings, PartnerCenterBox, PropertyUtilities);
+
+
+	//add a link to Partner Center, optionally including the StoreId for the title
+	auto PartnerCenterURL = [GDKTargetSettings]()
+	{
+		FString Result(TEXT("https://partner.microsoft.com/dashboard/"));
+		if (GDKTargetSettings.IsValid() && GDKTargetSettings->StoreId.Len() == 12 )
+		{
+			Result += FString::Printf(TEXT("products/%s/"), *GDKTargetSettings->StoreId);
+		}
+		return Result;
+	};
+	PartnerCenterBox->AddSlot()
+	.AutoWidth()
+	.VAlign(VAlign_Center)
+	.Padding(8,4)
+	[
+		SNew(SHyperlink)
+		.Text(LOCTEXT("PartnerCenterHyperLink", "Open Partner Center"))
+		.Style(FAppStyle::Get(), TEXT("NavigationHyperlink"))
+		.ToolTipText_Lambda([=]() { return FText::FromString(*PartnerCenterURL()); })
+		.OnNavigate_Lambda([=]() { FPlatformProcess::LaunchURL(*PartnerCenterURL(), nullptr, nullptr); })
+	];
+
+	//add a link to the stores, if the store id has been set
+	auto XboxStoreURL = [GDKTargetSettings]()
+	{
+		return GDKTargetSettings.IsValid() ? FString::Printf(TEXT("msxbox://game/?productid=%s"), *GDKTargetSettings->StoreId) : TEXT("");
+	};
+	auto WindowsStoreURL = [GDKTargetSettings]()
+	{
+		return GDKTargetSettings.IsValid() ? FString::Printf(TEXT("ms-windows-store://pdp/?productid=%s"), *GDKTargetSettings->StoreId) : TEXT("");
+	};
+	auto HasValidStoreId = [GDKTargetSettings]()
+	{
+		return (GDKTargetSettings.IsValid() && GDKTargetSettings->StoreId.Len() == 12 ) ? EVisibility::Visible : EVisibility::Collapsed; 
+	};
+	PartnerCenterBox->AddSlot()
+	.AutoWidth()
+	.VAlign(VAlign_Center)
+	.Padding(8,4)
+	[
+		SNew(SHyperlink)
+		.Text(LOCTEXT("XboxPCStoreHyperLink", "Show In Xbox PC App"))
+		.Style(FAppStyle::Get(), TEXT("NavigationHyperlink"))
+		.ToolTipText(LOCTEXT("XboxPCStoreHyperLinkTip", "Open this title in the Xbox PC App. The title must already have been published to the current sandbox."))
+		.OnNavigate_Lambda([=]() { FPlatformProcess::LaunchURL(*XboxStoreURL(), nullptr, nullptr); })
+		.Visibility(TAttribute<EVisibility>::Create(HasValidStoreId))
+	];
+	PartnerCenterBox->AddSlot()
+	.AutoWidth()
+	.VAlign(VAlign_Center)
+	.Padding(8,4)
+	[
+		SNew(SHyperlink)
+		.Text(LOCTEXT("WindowsStoreHyperLink", "Show In Windows Store"))
+		.Style(FAppStyle::Get(), TEXT("NavigationHyperlink"))
+		.ToolTipText(LOCTEXT("WindowsStoreHyperLinkTip", "Open this title in the Windows Store. The title must already have been published to the current sandbox."))
+		.OnNavigate_Lambda([=]() { FPlatformProcess::LaunchURL(*WindowsStoreURL(), nullptr, nullptr); })
+		.Visibility(TAttribute<EVisibility>::Create(HasValidStoreId))
+	];
+
+	// hook for adding any additional platform-specific items
+	AddAdditionalPlatformItems(PartnerCenterBox, GDKTargetSettings);
+
+
+	//add button to update product details from partner center
+	PartnerCenterBox->AddSlot()
+	.AutoWidth()
+	.VAlign(VAlign_Center)
+	.Padding(8,4)
+	[
+		SNew(SButton)
+		.Text(LOCTEXT("UpdateFromPartnerCenterLabel", "Update Partner Center Association"))
+		.ToolTipText(LOCTEXT("UpdateFromPartnerCenterTooltip", "Connect to Partner Center to display a list of products, and select one to associate with this game"))
+		.ButtonStyle(FAppStyle::Get(), "Button")
+		.Visibility_Lambda([this, GDKTargetSettings]() { return !IsGDKConfigured(GDKTargetSettings.Get()) ? EVisibility::Collapsed : EVisibility::Visible; })
+		.IsEnabled_Lambda([this]() { return !PartnerCenterProductWindow.IsValid() && !IGDKPlatformEditorModule::Get().IsQueryingPartnerCenter(); } )
+		.OnClicked_Lambda([this, GDKTargetSettings, PropertyUtilities]()
+		{
+			OnClickUpdateFromPartnerCenter(GDKTargetSettings.Get(), PropertyUtilities);
+			return FReply::Handled();
+		})
+	];
+
+	// hide the DLC category if the GDK DLC support plugin isn't enabled
+	TSharedPtr<IPlugin> GDKDLCPlugin = IPluginManager::Get().FindPlugin(TEXT("GDKDLCSupport"));
+	if (!GDKDLCPlugin.IsValid() || !GDKDLCPlugin->IsEnabled())
+	{
+		DetailBuilder.HideCategory("DLC");
+	}
+}
+
+
+
+void FGDKTargetSettingsDetails::AddGettingStartedItems(TWeakObjectPtr<UGDKTargetSettings> GDKTargetSettings, TSharedPtr<SHorizontalBox> PartnerCenterBox, const TSharedRef<IPropertyUtilities> PropertyUtilities)
+{
+	TSharedPtr<SHorizontalBox> GettingStartedBox;
+
+	PartnerCenterBox->AddSlot()
+	[
+		SNew(SBorder)
+		.BorderBackgroundColor(FColorList::SteelBlue)
+		.BorderImage(FAppStyle::GetBrush("ToolPanel.LightGroupBorder"))
+		.Visibility_Lambda([this, GDKTargetSettings]() { return IsGDKConfigured(GDKTargetSettings.Get()) ? EVisibility::Collapsed : EVisibility::Visible; })
+		.Padding(8.0f)
+		[
+			SAssignNew(GettingStartedBox, SHorizontalBox)
+
+			// Status icon
+			+ SHorizontalBox::Slot()
+			.AutoWidth()
+			.VAlign(VAlign_Center)
+			[
+				SNew(SImage)
+				.Image(FAppStyle::GetBrush("SettingsEditor.WarningIcon"))
+			]
+
+			// Notice
+			+ SHorizontalBox::Slot()
+			.FillWidth(1.0f)
+			.Padding(16.0f, 0.0f)
+			.VAlign(VAlign_Center)
+			[
+				SNew(STextBlock)
+				.ColorAndOpacity(FLinearColor::White)
+				.ShadowColorAndOpacity(FLinearColor::Black)
+				.ShadowOffset(FVector2D::UnitVector)
+				.Text(LOCTEXT("NotConfiguredLabel", "This project has not been configured for this platform yet. You can associate your game with Partner Center to get started."))
+			]
+
+			// Store Association Button
+			+ SHorizontalBox::Slot()
+			.AutoWidth()
+			.VAlign(VAlign_Center)
+			[
+				SNew(SButton)
+				.Text(LOCTEXT("ImportFromPartnerCenterLabel", "Associate With Partner Center"))
+				.ToolTipText(LOCTEXT("ImportFromPartnerCenterTooltip", "Connect to Partner Center to display a list of products, and select one to associate with this game"))
+				.ButtonStyle(FAppStyle::Get(), "Button")
+				.IsEnabled_Lambda([this]() { return !PartnerCenterProductWindow.IsValid() && !IGDKPlatformEditorModule::Get().IsQueryingPartnerCenter(); } )
+				.OnClicked_Lambda([this, GDKTargetSettings, PropertyUtilities]()
+				{
+					OnClickUpdateFromPartnerCenter(GDKTargetSettings.Get(), PropertyUtilities);
+					return FReply::Handled();
+				})
+			]
+		]
+	];
+
+	if (GettingStartedBox.IsValid())
+	{
+		AddAdditionalGettingStartedItems(GettingStartedBox, GDKTargetSettings, PropertyUtilities);
+	}
+}
+
+
+bool FGDKTargetSettingsDetails::IsGDKConfigured(UGDKTargetSettings* GDKTargetSettings) const
+{
+	return (GDKTargetSettings == nullptr || GDKTargetSettings->TitleId.Len() > 0 || GDKTargetSettings->PrimaryServiceConfigId.Len() > 0 || GDKTargetSettings->StoreId.Len() > 0 || GDKTargetSettings->PackageName.Len() > 0);
+}
+
+void FGDKTargetSettingsDetails::NotifyPropertyChanged(FProperty* Property, UObject* OwnerObject, const TSharedRef<class IPropertyUtilities> PropertyUtilities) const
+{
+	// signal that the property has changed - this makes sure the ini file is updated & saved
+	FNotifyHook* NotifyHook = PropertyUtilities->GetNotifyHook();
+	if (NotifyHook != nullptr)
+	{
+		TArray<const UObject*> NotifyTopLevelObjects;
+		NotifyTopLevelObjects.Add(OwnerObject);
+
+		FEditPropertyChain PropertyChain;
+		PropertyChain.AddHead(Property);
+
+		FPropertyChangedEvent ChangeEvent(Property, EPropertyChangeType::ValueSet, MakeArrayView(NotifyTopLevelObjects));
+		NotifyHook->NotifyPostChange(ChangeEvent, &PropertyChain);
+	}
+}
+
+void FGDKTargetSettingsDetails::NotifyPropertyChanged(FName PropertyName, UObject* OwnerObject, const TSharedRef<class IPropertyUtilities> PropertyUtilities) const
+{
+	FProperty* Property = FindFProperty<FProperty>( OwnerObject->GetClass(), PropertyName );
+	if (Property != nullptr)
+	{
+		NotifyPropertyChanged(Property, OwnerObject, PropertyUtilities);
+	}
+}
+
+
+
+
+
+
+
+
+
+
+
+
+
+
+
+
+
+
+void FGDKTargetSettingsDetails::OnClickUpdateFromPartnerCenter( UGDKTargetSettings* GDKTargetSettings, const TSharedRef<class IPropertyUtilities> PropertyUtilities )
+{
+	if (IGDKPlatformEditorModule::Get().IsQueryingPartnerCenter())
+	{
+		return;
+	}
+	check(!PartnerCenterProductWindow.IsValid());
+
+	// readability typedefs
+	typedef TSharedPtr<IGDKPlatformEditorModule::FPartnerCenterProduct> FProductPtr;
+	typedef SListView<FProductPtr> SProductListView;
+
+	// one row in the product selection listbox
+	auto OnGenerateRow = []( FProductPtr Product, const TSharedRef<STableViewBase>& OwnerTable )
+	{
+		const FString DisplayName = Product->FindRef(TEXT("PackageDisplayName"));
+		const FString PackageName = Product->FindRef(TEXT("PackageName"));
+		return SNew(STableRow<FProductPtr>, OwnerTable)
+		[
+			SNew(SHorizontalBox)
+			+SHorizontalBox::Slot()
+			.FillWidth(0.4)
+			[
+				SNew(STextBlock).Text(FText::FromString(DisplayName))
+			]
+
+			+SHorizontalBox::Slot()
+			.FillWidth(0.6)
+			[
+				SNew(STextBlock).Text(FText::FromString(PackageName))
+			]
+		];
+	};
+
+
+	// create the window
+	PartnerCenterProductWindow = SNew(SWindow)
+		.Title( LOCTEXT("PartnerCenterAssociate","Partner Center Association"))
+		.ClientSize(FVector2D(500,600))
+		.SupportsMaximize(false)
+		.SupportsMinimize(false)
+		.SizingRule( ESizingRule::UserSized )
+		;
+	PartnerCenterProductWindow->GetOnWindowClosedEvent().AddLambda( [this](const TSharedRef<SWindow>&)
+	{
+		IGDKPlatformEditorModule::Get().CancelQueryPartnerCenter();
+		PartnerCenterProductWindow.Reset();
+
+	});
+
+	// create the main window contents
+	TSharedPtr<SProductListView> ListView;
+	TSharedPtr<SWidgetSwitcher> WidgetSwitcher;
+
+	PartnerCenterProductWindow->SetContent(
+		SNew(SBorder)
+		.Padding(FMargin(8.0f, 8.0f))
+		[
+			SNew(SVerticalBox)
+			+SVerticalBox::Slot()
+			.FillHeight(1.0)
+			[
+				SAssignNew(WidgetSwitcher, SWidgetSwitcher)
+		
+				// progress spinner
+				+SWidgetSwitcher::Slot()
+				[
+					SNew(SBorder)
+					.HAlign(HAlign_Center)
+					.VAlign(VAlign_Center)
+					[
+						SNew(SThrobber)
+					]
+				]
+
+				// product selection listbox
+				+SWidgetSwitcher::Slot()
+				[
+					SAssignNew(ListView, SProductListView)
+					.ListItemsSource( IGDKPlatformEditorModule::Get().GetPartnerCenterProducts() )
+					.SelectionMode(ESelectionMode::SingleToggle)
+					.OnGenerateRow_Lambda(OnGenerateRow)
+				]			
+			]
+
+			// button area
+			+SVerticalBox::Slot()
+			.Padding(8,4)
+			.AutoHeight()
+			[
+				// ok button
+				SNew(SHorizontalBox)
+				+SHorizontalBox::Slot()
+				.Padding(5)
+				.AutoWidth()
+				[
+					SNew(SButton)
+					.ButtonStyle(FAppStyle::Get(), "PrimaryButton")
+					.TextStyle( FAppStyle::Get(), "DialogButtonText" )
+					.HAlign(HAlign_Center)
+					.VAlign(VAlign_Center)
+					.Text(LOCTEXT("OkButtonLabel", "Ok"))
+					.IsEnabled_Lambda( [ListView]() { return (ListView->GetNumItemsSelected() != 0) && !IGDKPlatformEditorModule::Get().IsQueryingPartnerCenter(); } )
+					.OnClicked_Lambda( [=, this]()
+					{
+						const FProductPtr Product = ListView->GetSelectedItems()[0];
+						UpdatePartnerCenterData( GDKTargetSettings, *Product, PropertyUtilities );
+						PartnerCenterProductWindow->RequestDestroyWindow();
+						PartnerCenterProductWindow.Reset();
+						return FReply::Handled();
+					})
+				]
+
+				// cancel button
+				+SHorizontalBox::Slot()
+				.Padding(5)
+				.AutoWidth()
+				[
+					SNew(SButton)
+					.ButtonStyle(FAppStyle::Get(), "Button")
+					.TextStyle( FAppStyle::Get(), "DialogButtonText" )
+					.HAlign(HAlign_Center)
+					.VAlign(VAlign_Center)
+					.Text(LOCTEXT("CancelButtonLabel", "Cancel"))
+					.OnClicked_Lambda( [this]()
+					{
+						PartnerCenterProductWindow->RequestDestroyWindow();
+						PartnerCenterProductWindow.Reset();
+						return FReply::Handled();
+					})
+				]
+			]
+		]
+	);
+
+	// show the store association dialog
+	IMainFrameModule& MainFrame = FModuleManager::LoadModuleChecked<IMainFrameModule>( "MainFrame" );
+	TSharedPtr<SWindow> ParentWindow = MainFrame.GetParentWindow();
+	if ( ParentWindow.IsValid() )
+	{
+		FSlateApplication::Get().AddModalWindow(PartnerCenterProductWindow.ToSharedRef(), ParentWindow.ToSharedRef(), true);
+		PartnerCenterProductWindow->ShowWindow();
+	}
+	else
+	{
+		FSlateApplication::Get().AddWindow(PartnerCenterProductWindow.ToSharedRef());
+	}
+
+	// refresh the list of products
+	auto OnPartnerCenterQueryComplete =  [this,GDKTargetSettings,ListView,WidgetSwitcher](bool bSuccess)
+	{
+		// show the list view and refresh it
+		WidgetSwitcher->SetActiveWidgetIndex(1);
+		ListView->RequestListRefresh();
+
+		if (bSuccess)
+		{
+			// find the currently selected item
+			const FProductPtr* CurrentProduct = IGDKPlatformEditorModule::Get().GetPartnerCenterProducts()->FindByPredicate( [GDKTargetSettings]( const FProductPtr Product )
+			{
+				const FString* StoreId = Product->Find(TEXT("StoreId"));
+				return (StoreId != nullptr) && StoreId->Equals(GDKTargetSettings->StoreId);
+			});
+			if (CurrentProduct && CurrentProduct->IsValid())
+			{
+				ListView->SetSelection( *CurrentProduct );
+			}
+		}
+		else
+		{
+			// close the window
+			PartnerCenterProductWindow->RequestDestroyWindow();
+			PartnerCenterProductWindow.Reset();
+		}
+	};
+
+	// show the spinner widget and start the query
+	WidgetSwitcher->SetActiveWidgetIndex(0);
+	IGDKPlatformEditorModule::Get().PartnerCenterQueryProductsAsync(OnPartnerCenterQueryComplete);
+}
+
+
+
+void FGDKTargetSettingsDetails::UpdatePartnerCenterData( UGDKTargetSettings* GDKTargetSettings, const IGDKPlatformEditorModule::FPartnerCenterProduct& Product, const TSharedRef<class IPropertyUtilities> PropertyUtilities )
+{
+	auto UpdateProperty = [this, Product, GDKTargetSettings, PropertyUtilities]( FString& Property, const FString& Key, FName MemberName )
+	{
+		const FString* Value = Product.Find(Key);
+		if (Value && !Value->IsEmpty() && !Value->Equals(Property))
+		{
+			Property = *Value;
+			NotifyPropertyChanged(MemberName, GDKTargetSettings, PropertyUtilities);
+		}
+	};
+
+	UpdateProperty( GDKTargetSettings->TitleId,                TEXT("TitleId"),                GET_MEMBER_NAME_CHECKED(UGDKTargetSettings, TitleId) );
+	UpdateProperty( GDKTargetSettings->MSAAppId,               TEXT("MSAAppId"),               GET_MEMBER_NAME_CHECKED(UGDKTargetSettings, MSAAppId) );
+	UpdateProperty( GDKTargetSettings->StoreId,                TEXT("StoreId"),                GET_MEMBER_NAME_CHECKED(UGDKTargetSettings, StoreId) );
+	UpdateProperty( GDKTargetSettings->PackageName,            TEXT("PackageName"),            GET_MEMBER_NAME_CHECKED(UGDKTargetSettings, PackageName) );
+	UpdateProperty( GDKTargetSettings->PrimaryServiceConfigId, TEXT("PrimaryServiceConfigId"), GET_MEMBER_NAME_CHECKED(UGDKTargetSettings, PrimaryServiceConfigId) );
+	UpdateProperty( GDKTargetSettings->PublisherName,          TEXT("PublisherName"),          GET_MEMBER_NAME_CHECKED(UGDKTargetSettings, PublisherName) );
+	UpdateProperty( GDKTargetSettings->DefaultDisplayName,     TEXT("PackageDisplayName"),     GET_MEMBER_NAME_CHECKED(UGDKTargetSettings, DefaultDisplayName) );
+	UpdateProperty( GDKTargetSettings->ProductId,              TEXT("ProductId"),              GET_MEMBER_NAME_CHECKED(UGDKTargetSettings, ProductId) );
+	
+	PropertyUtilities->ForceRefresh();
+}
+
+
+
+
+
+
+
+
+FString FGDKCultureResourceDetails::GetPickerPath()
+{
+	return FEditorDirectories::Get().GetLastDirectory(ELastDirectory::GENERIC_OPEN);
+}
+
+bool FGDKCultureResourceDetails::HandlePostExternalIconCopy(const FString& InChosenImage)
+{
+	FEditorDirectories::Get().SetLastDirectory(ELastDirectory::GENERIC_OPEN, FPaths::GetPath(InChosenImage));
+	return true;
+}
+
+TSharedPtr<IPropertyHandle> FGDKCultureResourceDetails::GetDLCNameProperty( TSharedPtr<IPropertyHandle> InStructPropertyHandle ) const
+{
+	TSharedPtr<IPropertyHandle> ParentPropertyHandle = InStructPropertyHandle->GetParentHandle();
+
+	// the parent of the struct has a field called 'DLCName' (such as FGDKDLCConfiguration::DefaultStringResources)
+	TSharedPtr<IPropertyHandle> ChildHandle = ParentPropertyHandle->GetChildHandle("DLCName", false);
+	if (ChildHandle.IsValid())
+	{
+		return ChildHandle;
+	}
+
+	// the great-grandparent of the struct has a field called 'DLCName' (such as FGDKDLCConfiguration::PerCultureResources[N])
+	ParentPropertyHandle = ParentPropertyHandle->GetParentHandle();
+	ParentPropertyHandle = ParentPropertyHandle.IsValid() ? ParentPropertyHandle->GetParentHandle() : nullptr;
+	if (ParentPropertyHandle.IsValid())
+	{
+		ChildHandle = ParentPropertyHandle->GetChildHandle("DLCName", false);
+		if (ChildHandle.IsValid())
+		{
+			return ChildHandle;
+		}
+	}
+
+	return nullptr;
+}
+
+void FGDKCultureResourceDetails::CustomizeChildren(TSharedRef<IPropertyHandle> InStructPropertyHandle, class IDetailChildrenBuilder& ChildBuilder, IPropertyTypeCustomizationUtils& StructCustomizationUtils)
+{
+	// add all the child properties (ApplicationDisplayName & ApplicationDescription)
+	uint32 NumChildren;
+	if (InStructPropertyHandle->GetNumChildren(NumChildren) == FPropertyAccess::Success)
+	{
+		for (uint32 Index = 0; Index < NumChildren; Index++)
+		{
+			TSharedPtr<IPropertyHandle> ChildProperty = InStructPropertyHandle->GetChildHandle(Index);
+			ChildBuilder.AddProperty(ChildProperty.ToSharedRef());
+		}
+	}
+
+	// add the button to bring up the image picker dialog
+	TSharedPtr<IPropertyHandle> ParentPropertyHandle = InStructPropertyHandle->GetParentHandle();
+	TSharedPtr<IPropertyHandle> DLCNameHandle = GetDLCNameProperty(InStructPropertyHandle);
+	TSharedPtr<IPropertyHandle> StageIdHandle = ParentPropertyHandle->GetChildHandle("StageId", false);
+	FText PickImgLabel = StageIdHandle.IsValid() ? LOCTEXT("PickLocImgsLabel", "Application Images") : LOCTEXT("PickDefImgsLabel", "Application Images");
+	FText PickImgToolTip = StageIdHandle.IsValid() ? LOCTEXT("PickLocImgsTip", "Configure the language-specific images for the package") : LOCTEXT("PickDefImgsTip", "Configure the default images for the package");
+
+	ChildBuilder.AddCustomRow(FText::GetEmpty())
+		.NameWidget
+		[
+			SNew(STextBlock)
+				.Text(PickImgLabel)
+				.Font(StructCustomizationUtils.GetRegularFont())
+				.ToolTipText(PickImgToolTip)
+		]
+		.ValueWidget
+		[
+			SNew(SButton)
+			.Text(LOCTEXT("PickAppImgBtnLabel","Configure..."))
+			.ToolTipText(PickImgToolTip)
+			.ButtonStyle(FAppStyle::Get(), "Button")
+			.OnClicked_Lambda([this, InStructPropertyHandle]()
+			{
+				OnClickConfigureImages(InStructPropertyHandle);
+				return FReply::Handled();
+			})
+			.IsEnabled_Lambda([=]()
+			{
+				// if these images belong to DLC, cannot click until the DLC name is selected
+				FString DLCName;
+				if (DLCNameHandle.IsValid() && (DLCNameHandle->GetValue(DLCName) == FPropertyAccess::Result::Success) && DLCName.IsEmpty())
+				{
+					return false;
+				}
+
+				// default culture always clickable
+				if (!StageIdHandle.IsValid())
+				{
+					return true;
+				}
+
+				// otherwise, we must have stage id before this can be clicked
+				FString StageId;
+				return (StageIdHandle->GetValue(StageId) == FPropertyAccess::Result::Success) && !StageId.IsEmpty(); 
+			})
+		];
+}
+
+void FGDKCultureResourceDetails::OnClickConfigureImages(TSharedRef<IPropertyHandle> InStructPropertyHandle)
+{
+	// Look up the selected target settings object & find the platform name
+	const TCHAR* PlatformName = nullptr;
+	UGDKTargetSettings* TargetSettings = nullptr;
+	TArray<UObject*> OuterObjects;
+	InStructPropertyHandle->GetOuterObjects(OuterObjects);
+	for (UObject* OuterObject : OuterObjects)
+	{
+		TargetSettings = Cast<UGDKTargetSettings>(OuterObject);
+		if (TargetSettings != nullptr)
+		{
+			PlatformName = TargetSettings->GetPlatformName();
+			break;
+		}
+	}
+	if (PlatformName == nullptr)
+	{
+		return;
+	}
+
+	const FDataDrivenPlatformInfo& PlatformInfo = FDataDrivenPlatformInfoRegistry::GetPlatformInfo(FName(PlatformName));
+	bool bPlatformExtension = PlatformInfo.bIsConfidential;
+	FString IniPlatformName = PlatformInfo.IniPlatformName.ToString();
+
+	// look up the Culture Id
+	TSharedPtr<IPropertyHandle> ParentPropertyHandle = InStructPropertyHandle->GetParentHandle();
+	TSharedPtr<IPropertyHandle> StageIdHandle = ParentPropertyHandle->GetChildHandle("StageId", false);
+
+	FString OwnerName;
+	FString CultureId;
+	bool bIsDefaultCulture = true;
+	if (StageIdHandle.IsValid())
+	{
+		FString StageId;
+		if ((StageIdHandle->GetValue(StageId) == FPropertyAccess::Result::Success) && !StageId.IsEmpty())
+		{
+			// look up the CultureId for the given StageId
+			const FString* CultureIdPtr = TargetSettings->StageIdOverrides.Find(StageId);
+			CultureId = CultureIdPtr ? *CultureIdPtr : StageId;
+			bIsDefaultCulture = false;
+		}
+		else
+		{
+			// If a stage id isn't set we can't display images
+			return;
+		}
+	}
+
+
+	// lookup DLC
+	TSharedPtr<IPlugin> DLCPlugin;
+	TSharedPtr<IPropertyHandle> DLCNameHandle = GetDLCNameProperty(InStructPropertyHandle);
+	if (DLCNameHandle.IsValid())
+	{
+		FString DLCName;
+		if ((DLCNameHandle->GetValue(DLCName) == FPropertyAccess::Result::Success) && !DLCName.IsEmpty())
+		{
+			DLCPlugin = IPluginManager::Get().FindPlugin(DLCName);
+		}
+
+		// If a DLC name isn't set or the plugin isn't known, we can't display images
+		if (!DLCPlugin.IsValid())
+		{
+			UE_LOG(LogTemp, Error, TEXT("Plugin not found or not specified. %s"), *DLCName);
+			return;
+		}
+	}
+
+	// look up the base directory for the MSGamingSupport plugin to get the default images
+	FString MSGamingSupportResourcePath = FPaths::EngineDir() / GDKTargetSettingsDetailsConstants::EngineResourcePath;
+	TSharedPtr<IPlugin> MSGamingSupportPlugin = IPluginManager::Get().FindPlugin(TEXT("MSGamingSupport"));
+	if (MSGamingSupportPlugin.IsValid())
+	{
+		FString PluginBaseDir = FPaths::GetPath(MSGamingSupportPlugin->GetDescriptorFileName());
+		MSGamingSupportResourcePath = PluginBaseDir / GDKTargetSettingsDetailsConstants::EngineResourcePath;
+	}
+
+	// Setup default image path, use the engine copy if we're already dealing with the default culture, otherwise use the default culture
+	FString DefaultResourcePath = "";
+	if (bIsDefaultCulture)
+	{
+		if (bPlatformExtension)
+		{
+			DefaultResourcePath = FPaths::EnginePlatformExtensionDir(PlatformName) / GDKTargetSettingsDetailsConstants::EngineResourcePath;
+		}
+		else
+		{
+			DefaultResourcePath = MSGamingSupportResourcePath;
+		}
+	}
+	else
+	{
+		if (bPlatformExtension)
+		{
+			DefaultResourcePath = FPaths::ProjectPlatformExtensionDir(PlatformName) / GDKTargetSettingsDetailsConstants::BuildResourcePath;
+		}
+		else
+		{
+			DefaultResourcePath = FPaths::ProjectDir() / TEXT("Build") / IniPlatformName / TEXT("MSGaming/Resources");
+		}
+	}
+	DefaultResourcePath = FPaths::ConvertRelativePathToFull(DefaultResourcePath);
+
+	// The target path is always setup to the culture of the control or the root of the resource path if it's the default culture control
+	FString TargetResourcePath = "";
+	if (DLCPlugin.IsValid())
+	{
+		if (bPlatformExtension)
+		{
+			TargetResourcePath = DLCPlugin->GetBaseDir() / TEXT("Platforms") / PlatformName / GDKTargetSettingsDetailsConstants::BuildResourcePath;
+		}
+		else
+		{
+			TargetResourcePath = DLCPlugin->GetBaseDir() / TEXT("Build") / IniPlatformName / TEXT("MSGaming/Resources");
+		}
+	}
+	else
+	{
+		if (bPlatformExtension)
+		{
+			TargetResourcePath = FPaths::ProjectPlatformExtensionDir(PlatformName) / GDKTargetSettingsDetailsConstants::BuildResourcePath;
+		}
+		else
+		{
+			TargetResourcePath = FPaths::ProjectDir() / TEXT("Build") / IniPlatformName / TEXT("MSGaming/Resources");
+		}
+
+	}
+	if (!bIsDefaultCulture)
+	{
+		TargetResourcePath /= CultureId;
+	}
+	TargetResourcePath = FPaths::ConvertRelativePathToFull(TargetResourcePath);
+
+
+
+	// prepare window title, showing the selected culture id if appropriate
+	FText WindowTitle = LOCTEXT("PickAppImgDefault", "Default Packaging Images");
+	if (!bIsDefaultCulture)
+	{
+		const FText CultureIdText = FText::FromString(CultureId);
+		WindowTitle = FText::Format(LOCTEXT("PickAppImgLocale", "Packaging Images For Locale \"{0}\""), CultureIdText);
+	}
+
+	// create the window
+	TSharedPtr<SWindow> Window = SNew(SWindow)
+		.Title(WindowTitle)
+		.ClientSize(FVector2D(600,800))
+		.SizingRule(ESizingRule::UserSized);
+
+	TSharedPtr<SScrollBox> ImagesBox;
+	Window->SetContent(
+		SNew(SBorder)
+		.Padding(FMargin(8.0f, 8.0f))
+		[
+			// product selection listbox
+			SNew(SVerticalBox)
+			+SVerticalBox::Slot()
+			.FillHeight(1.0)
+			.Padding(5, 5)
+			[
+				SAssignNew(ImagesBox, SScrollBox)
+			]
+
+			// button area
+			+SVerticalBox::Slot()
+			.Padding(8,4)
+			.AutoHeight()
+			[
+				SNew(SHorizontalBox)
+
+				// open resource images folder
+				+SHorizontalBox::Slot()
+				.HAlign(HAlign_Left)
+				.VAlign(VAlign_Bottom)
+				[
+					SNew(SHyperlink)
+					.Text(LOCTEXT("OpenFolder", "Open Folder"))
+					.ToolTipText(FText::FromString(TargetResourcePath))
+					.OnNavigate_Lambda( [TargetResourcePath]() { IPlatformFile::GetPlatformPhysical().CreateDirectoryTree(*TargetResourcePath); FPlatformProcess::ExploreFolder(*TargetResourcePath); })
+				]
+
+				+SHorizontalBox::Slot()
+				.HAlign(HAlign_Right)
+				.VAlign(VAlign_Bottom)
+				[
+					SNew(SButton)
+					.Text(LOCTEXT("Close", "Close"))
+					.OnClicked_Lambda( [=]()
+					{
+						Window->RequestDestroyWindow();
+						return FReply::Handled();
+					})
+				]
+			]
+		]
+	);
+
+
+	// Application small logo widget
+	const FText AppSmLogoDesc(LOCTEXT("AppSmLogoLabel", "Application 44x44 Small Logo"));
+	const FString AppSmLogo_TargetImagePath(TargetResourcePath / "SmallLogo.png");
+	const FString AppSmLogo_DefaultImagePath(DefaultResourcePath / "SmallLogo.png");
+	AddImagePickerRow(ImagesBox, AppSmLogoDesc, AppSmLogo_DefaultImagePath, AppSmLogo_TargetImagePath, 44, 44);
+
+	// Store logo widget
+	const FText StoreLogoDesc(LOCTEXT("StoreLogoLabel", "Store 100x100 Logo"));
+	const FString StoreLogo_TargetImagePath(TargetResourcePath / "StoreLogo.png");
+	const FString StoreLogo_DefaultImagePath(DefaultResourcePath / "StoreLogo.png");
+	AddImagePickerRow(ImagesBox, StoreLogoDesc, StoreLogo_DefaultImagePath, StoreLogo_TargetImagePath, 100, 100);
+
+ 	// Application logo widget
+ 	const FText AppLogoDesc(LOCTEXT("AppLogoLabel", "Application 150x150 Logo"));
+ 	const FString AppLogo_TargetImagePath(TargetResourcePath / "Logo.png");
+ 	const FString AppLogo_DefaultImagePath(DefaultResourcePath / "Logo.png");
+	AddImagePickerRow(ImagesBox, AppLogoDesc, AppLogo_DefaultImagePath, AppLogo_TargetImagePath, 150, 150);
+
+	// Application large logo widget
+ 	const FText AppLgLogoDesc(LOCTEXT("AppLargeLogoLabel", "Application 480x480 Large Logo"));
+ 	const FString AppLgLogo_TargetImagePath(TargetResourcePath / "Square480x480Logo.png");
+ 	const FString AppLgLogo_DefaultImagePath(DefaultResourcePath / "Square480x480Logo.png");
+	AddImagePickerRow(ImagesBox, AppLgLogoDesc, AppLgLogo_DefaultImagePath, AppLgLogo_TargetImagePath, 480, 480);
+
+ 	// Application splash screen widget
+ 	const FText AppSplashDesc(LOCTEXT("AppSplashLabel", "Application 1920x1080 Splash Screen"));
+ 	const FString AppSplash_TargetImagePath(TargetResourcePath / "SplashScreen.png");
+ 	const FString AppSplash_DefaultImagePath(DefaultResourcePath / "SplashScreen.png");
+	AddImagePickerRow(ImagesBox, AppSplashDesc, AppSplash_DefaultImagePath, AppSplash_TargetImagePath, 1920, 1080);
+
+
+	// show the window
+	IMainFrameModule& MainFrame = FModuleManager::LoadModuleChecked<IMainFrameModule>("MainFrame");
+	TSharedPtr<SWindow> ParentWindow = MainFrame.GetParentWindow();
+	if (ParentWindow.IsValid())
+	{
+		FSlateApplication::Get().AddModalWindow(Window.ToSharedRef(), ParentWindow.ToSharedRef());
+	}
+	else
+	{
+		FSlateApplication::Get().AddWindow(Window.ToSharedRef());
+	}
+}
+
+
+
+
+void FGDKCultureResourceDetails::AddImagePickerRow(TSharedPtr<class SScrollBox> ImagesBox, const FText& ImageDescription, const FString& DefaultImagePath, const FString& TargetImagePath, int32 ImageWidth, int32 ImageHeight)
+{
+	TArray<FString> ImageExtensions;
+	ImageExtensions.Add(TEXT("png"));
+
+	ImagesBox->AddSlot()
+		[
+			SNew(STextBlock)
+			.Text(ImageDescription)
+		];
+
+	ImagesBox->AddSlot()
+	.Padding(4,0,4,16)
+	[
+		SNew(SExternalImageReference, DefaultImagePath, TargetImagePath)
+		.FileDescription(ImageDescription)
+		.RequiredSize(FIntPoint(ImageWidth, ImageHeight))
+		.OnGetPickerPath(FOnGetPickerPath::CreateSP(this, &FGDKCultureResourceDetails::GetPickerPath))
+		.OnPostExternalImageCopy(FOnPostExternalImageCopy::CreateSP(this, &FGDKCultureResourceDetails::HandlePostExternalIconCopy))
+		.DeleteTargetWhenDefaultChosen(true)
+		.FileExtensions(ImageExtensions)
+		.DeletePreviousTargetWhenExtensionChanges(true)
+	];
+}
+
+#undef LOCTEXT_NAMESPACE
diff --git a/Engine/Source/Developer/Microsoft/GDK/GDKPlatformEditor/Public/GDKPlatformEditorModule.h b/Engine/Source/Developer/Microsoft/GDK/GDKPlatformEditor/Public/GDKPlatformEditorModule.h
new file mode 100644
--- /dev/null
+++ b/Engine/Source/Developer/Microsoft/GDK/GDKPlatformEditor/Public/GDKPlatformEditorModule.h
@@ -0,0 +1,27 @@
+// Copyright Epic Games, Inc. All Rights Reserved.
+
+#pragma once
+
+#include "CoreMinimal.h"
+#include "Modules/ModuleInterface.h"
+#include "Modules/ModuleManager.h"
+#include "Templates/Function.h"
+#include "Templates/SharedPointer.h"
+#include "Containers/Map.h"
+
+class IGDKPlatformEditorModule : public IModuleInterface
+{
+public:	
+	static inline IGDKPlatformEditorModule& Get()
+	{
+		return FModuleManager::LoadModuleChecked<IGDKPlatformEditorModule>("GDKPlatformEditor");
+	}
+
+	typedef TMap<FString,FString> FPartnerCenterProduct;
+
+	virtual void PartnerCenterQueryProductsAsync( TFunction<void(bool)> OnComplete ) = 0;
+	virtual const TArray<TSharedPtr<FPartnerCenterProduct>>* GetPartnerCenterProducts() = 0;
+
+	virtual bool IsQueryingPartnerCenter() const = 0;
+	virtual void CancelQueryPartnerCenter() = 0;
+};
diff --git a/Engine/Source/Developer/Microsoft/GDK/GDKPlatformEditor/Public/GDKTargetSettings.h b/Engine/Source/Developer/Microsoft/GDK/GDKPlatformEditor/Public/GDKTargetSettings.h
new file mode 100644
--- /dev/null
+++ b/Engine/Source/Developer/Microsoft/GDK/GDKPlatformEditor/Public/GDKTargetSettings.h
@@ -0,0 +1,440 @@
+// Copyright Epic Games, Inc. All Rights Reserved.
+
+/*=============================================================================
+	GDKTargetSettings.h: Declares the UGDKTargetSettings class.
+=============================================================================*/
+
+#pragma once
+
+#include "CoreMinimal.h"
+#include "UObject/Object.h"
+#include "GDKTargetSettings.generated.h"
+
+#define UE_API GDKPLATFORMEDITOR_API
+
+UENUM()
+enum class EGDKForegroundColor : uint8
+{
+	/** Use light colored text to contrast against a dark background. */
+	Light,
+
+	/** Use dark colored text to contrast against a light background. */
+	Dark,
+};
+
+
+USTRUCT()
+struct FGDKPackageStringResources
+{
+	GENERATED_USTRUCT_BODY()
+
+	/**
+	* Friendly application display name. This value will be displayed on game tiles in the dashboard.
+	*/
+	UPROPERTY(EditAnywhere, Config, Category=Content)
+	FString ApplicationDisplayName;
+
+	/**
+	* Friendly application description. This value will be displayed on game tiles in the game collection.
+	*/
+	UPROPERTY(EditAnywhere, Config, Category=Content)
+	FString ApplicationDescription;
+};
+
+USTRUCT()
+struct FGDKLocalizedPackageResources
+{
+	GENERATED_USTRUCT_BODY()
+
+	/**
+	* Which CultureToStage entry these resources should apply to. See "Languages To Package" in Packaging settings
+	*/
+	UPROPERTY(EditAnywhere, Config, Category=Content, Meta=(GetOptions="GetKnownCultureStageIds", DisplayName="Language"))
+	FString StageId;
+
+	/**
+	* The package localized strings.
+	*/
+	UPROPERTY(EditAnywhere, Config, Category=Content)
+	FGDKPackageStringResources CultureStringResources;
+};
+
+
+
+
+
+UENUM()
+enum class EGDKCustomChunkMappingType : uint8
+{
+	Main,
+	Optional,
+};
+
+USTRUCT()
+struct FGDKCustomChunkMapping
+{
+	GENERATED_USTRUCT_BODY()
+
+	UPROPERTY(EditAnywhere, Config, Category=Packaging, Meta=(ClampMin=0))
+	int32 ChunkId = 0;
+
+	// relative content path: e.g. shootergame/content/movies/\*  or  shootergame/content/movies/movie1.mp4
+	UPROPERTY(EditAnywhere, Config, Category=Packaging)
+	FString Pattern;
+
+	// not currently wired up, so keeping at default
+	UPROPERTY()
+	EGDKCustomChunkMappingType Type = EGDKCustomChunkMappingType::Main;
+};
+
+
+
+
+USTRUCT()
+struct FGDKIntelligentDeliveryChunk
+{
+	GENERATED_USTRUCT_BODY()
+
+	// Unreal chunk identifier
+	UPROPERTY(EditAnywhere, Config, Category=Packaging, Meta=(ClampMin=0))
+	int32 ChunkId = 0;
+
+	// List of tags, separated by semicolons, (using # to define groups if Features are also used). Does not apply when 'required for launch' is set. See the GDK documentation for details
+	UPROPERTY(EditAnywhere, Config, Category=Packaging)
+	FString Tags;
+
+	// List of languages from CultureToStage. See "Languages To Package" in Packaging settings
+	UPROPERTY(EditAnywhere, Config, Category=Packaging, Meta=(GetOptions="GetKnownCultureStageIds", DisplayName="Languages"))
+	TArray<FString> StageIds;
+
+	// List of devices that this chunk is available for.
+	UPROPERTY(EditAnywhere, Config, Category=Packaging, Meta=(GetOptions="GetKnownDeviceTypes"))
+	TArray<FString> Devices;
+
+	// Whether this chunk must be installed before the game can be launched.
+	UPROPERTY(EditAnywhere, Config, Category=Packaging)
+	bool bRequiredForLaunch = false;
+
+};
+
+USTRUCT()
+struct FGDKLocalizedIntelligentDeliveryFeatureResource
+{
+	GENERATED_USTRUCT_BODY()
+
+	// Locale from CultureToStage. See "Languages To Package" in Packaging settings
+	UPROPERTY(EditAnywhere, Config, Category=Content, Meta=(GetOptions="GetKnownCultureStageIds", DisplayName="Language"))
+	FString StageId;
+
+	// Friendly install feature display name. This value will be displayed to the user in the shell
+	UPROPERTY(EditAnywhere, Config, Category=Content, AdvancedDisplay)
+	FString DisplayName;
+};
+
+
+USTRUCT()
+struct FGDKIntelligentDeliveryFeature
+{
+	GENERATED_USTRUCT_BODY()
+
+	UPROPERTY(EditAnywhere, Config, Category=Feature)
+	FString Id;
+
+	// Semicolon-separated list of chunk tags. See the GDK documentation for details
+	UPROPERTY(EditAnywhere, Config, Category=Feature)
+	FString Tags;
+
+	// Whether this feature is hidden from the user. See the GDK documentation for details. 
+	UPROPERTY(EditAnywhere, Config, Category=Feature)
+	bool bHidden = true;
+
+	// Friendly install feature display name. This value will be displayed to the user in the shell.
+	UPROPERTY(EditAnywhere, Config, Category=Content, Meta=(EditCondition="!bHidden"))
+	FString DefaultDisplayName;
+
+	// Optional per-language resources that will be used for this feature
+	UPROPERTY(EditAnywhere, Config, Category=Content, DisplayName="Per-Language Resources", Meta=(EditCondition="!bHidden", TitleProperty="StageId"))
+	TArray<FGDKLocalizedIntelligentDeliveryFeatureResource> PerCultureResources;
+};
+
+
+USTRUCT()
+struct FGDKIntelligentDeliveryRecipe
+{
+	GENERATED_USTRUCT_BODY()
+
+	UPROPERTY(EditAnywhere, Config, Category=Recipe)
+	FString Id;
+
+	// List of features to include, separated by semicolons. See the GDK documentation for details
+	UPROPERTY(EditAnywhere, Config, Category=Recipe, Meta=(GetOptions="GetKnownIntelligentDeliveryFeatures"))
+	TArray<FString> IncludedFeatures;
+
+	// List of devices that this recipe is available for.
+	UPROPERTY(EditAnywhere, Config, Category=Recipe, Meta=(GetOptions="GetKnownDeviceTypes"))
+	TArray<FString> Devices;
+
+	// Premium product recipe. Only installed if the user owns one of these products. See the GDK documentation for details
+	UPROPERTY(EditAnywhere, Config, Category=Recipe)
+	TArray<FString> StoreIds;
+};
+
+USTRUCT()
+struct FGDKDLCPackageStringResources
+{
+	GENERATED_USTRUCT_BODY();
+
+	// Locale from CultureToStage. See "Languages To Package" in Packaging settings
+	UPROPERTY(EditAnywhere, Config, Category=Content, Meta=(GetOptions="GetKnownCultureStageIds", DisplayName="Language"))
+	FString StageId;
+
+	// The localized strings
+	UPROPERTY(EditAnywhere, Config, Category=Content)
+	FGDKPackageStringResources CultureStringResources;
+};
+
+
+USTRUCT()
+struct FGDKDLCConfiguration
+{
+	GENERATED_USTRUCT_BODY()
+
+	/* Name of the explicity-loaded DLC plugin this configuration applies to */
+	UPROPERTY(EditAnywhere, Config, Category=DLC, Meta=(GetOptions="GetKnownDLCNames", DisplayName="Plugin"))
+	FString DLCName;
+
+	/* Store Id. Must match Partner Center configuration. */
+	UPROPERTY(EditAnywhere, Config, Category=PartnerCenter)
+	FString StoreId;
+
+	/* Name of the product in Partner Center. Must match Partner Center configuration */
+	UPROPERTY(EditAnywhere, Config, Category=PartnerCenter, DisplayName="Product Name")
+	FString DefaultDisplayName;
+
+	/* Package/Identity/Name. Must match Partner Center configuration */
+	UPROPERTY(EditAnywhere, Config, Category=PartnerCenter, DisplayName="Package/Identity/Name")
+	FString PackageName;
+
+	/* Package version in the form 0.0.0.0. Leave blank to inherit the main package version */
+	UPROPERTY(EditAnywhere, Config, Category=Packaging)
+	FString PackageVersion;
+
+
+	/* The package default strings. */
+	UPROPERTY(EditAnywhere, Config, Category=Content)
+	FGDKPackageStringResources DefaultStringResources;
+
+	/* The per culture images, icons, and strings describing a package. */
+	UPROPERTY(EditAnywhere, Config, Category=Content, DisplayName="Per-Language Resources", Meta=(TitleProperty="StageId"))
+	TArray<FGDKDLCPackageStringResources> PerCultureResources;
+};
+
+
+USTRUCT()
+struct FGDKUriActivationProtocol
+{
+	GENERATED_USTRUCT_BODY()
+
+	/* e.g. MyGameUGC
+	 * Protocol that is registered with the OS. This will be what is launched by the user via MyGameUGC://<payload>
+	 * note that this will be system-wide, so should be clearly unique for your title.
+	 */
+	UPROPERTY(EditAnywhere, Config, Category=ActivationProtocol, Meta=(DisplayName="OS-Registered Protocol", MaxLength=39))
+	FString RegisteredProtocol;
+
+	/* e.g. ugc://
+	 * Prefix for how this protocol's payload will be represented in FCoreDelegates::OnActivatedByProtocol as ugc://<payload> 
+	 * note that this is only visible to your game, so can be a generic name.
+	 * this can be combined with other <platform>.*UriPrefix CVars to craft platform-agnostic OnActivatedByProtocol handler code.
+	 */
+	UPROPERTY(EditAnywhere, Config, Category=ActivationProtocol)
+	FString InGameUriPrefix;
+};
+
+UENUM()
+enum class EGDKCustomInstallActionType : uint8
+{
+	Invalid UMETA(Hidden, DisplayName="..."),
+	Install,
+	Repair,
+	Uninstall UMETA(ToolTip="Note: there is a bug in Windows 10 that prevents the uninstall action from being run. It is fixed in Windows 11"),
+};
+
+USTRUCT()
+struct FGDKCustomInstallCommand
+{
+	GENERATED_USTRUCT_BODY()
+
+	/* .exe or .msi file to run */
+	UPROPERTY(EditAnywhere, Config, Category=Action, Meta=(FilePathFilter="Commands (*.exe;*.msi)|*.exe;*.msi|All Files (*.*)|*.*"))
+	FFilePath Command;
+
+	/* command line arguments to pass to the given command. For .msi files this can be blank */
+	UPROPERTY(EditAnywhere, Config, Category=Action)
+	FString Arguments;
+};
+
+USTRUCT()
+struct FGDKCustomInstallAction
+{
+	GENERATED_USTRUCT_BODY()
+
+	/* name of the item to be installed/repaired/uninstalled. SourceFolder will be staged into a subfolder with this name */
+	UPROPERTY(EditAnywhere, Config, Category=Action)
+	FString Name;
+
+	/* Source folder for all the custom install commands. The contents of this directory will be staged automatically */
+	UPROPERTY(EditAnywhere, Config, Category=Actions, Meta=(RelativePath))
+	FDirectoryPath SourceFolder;
+
+	/* Relative destination folder when staging, relative to Custom Install Staging Root. Leave blank to stage the custom install action into a subfolder  */
+	UPROPERTY(EditAnywhere, Config, Category=Action)
+	FString TargetFolder;
+
+	/* The commands to run for each step of the installation */
+	UPROPERTY(EditAnywhere, Config, Category=Action)
+	TMap<EGDKCustomInstallActionType,FGDKCustomInstallCommand> Commands;
+};
+
+
+
+/**
+ * Implements the settings for the GDK target platform.
+  * Note: Properties must use Config not GlobalConfig because the properties need to be serialized into the specialization's own ini file.
+ */
+UCLASS(MinimalAPI, config = Engine, defaultconfig, Abstract)
+class UGDKTargetSettings
+	: public UObject
+{
+public:
+
+	GENERATED_UCLASS_BODY()
+
+	virtual const TCHAR* GetPlatformName() const PURE_VIRTUAL(GetPlatformName, return nullptr;);
+
+	/* Inform the OS of what color scheme best contrasts with your application's tile images. */
+	UPROPERTY(EditAnywhere, Config, Category=Content, Meta=(DisplayName="Application Foreground Text"))
+	EGDKForegroundColor ApplicationForegroundColor;
+
+	/* Inform the OS what background color to use for your application tiles. */
+	UPROPERTY(EditAnywhere, Config, Category=Content)
+	FColor ApplicationBackgroundColor;
+
+	/* The package default strings. */
+	UPROPERTY(EditAnywhere, Config, Category=Content)
+	FGDKPackageStringResources CultureStringResources;
+
+	/* The per culture images, icons, and strings describing a package. */
+	UPROPERTY(EditAnywhere, Config, Category=Content, DisplayName="Per-Language Resources", meta = (TitleProperty="StageId"))
+	TArray<FGDKLocalizedPackageResources> PerCultureResources;
+
+	/* Mapping from "Languages To Package" in Packaging settings, to "Console-supported languages" in the GDK documentation. Leave empty to keep 1:1 mapping */
+	UPROPERTY(EditAnywhere, Config, Category=Content, AdvancedDisplay, Meta=(GetKeyOptions="GetKnownCultureStageIds", DisplayName="Language/Locale Overrides"))
+	TMap<FString,FString> StageIdOverrides;
+
+
+	/* Whether to use the Simplified User Model where the OS will maintain a primary user for the title and the title will be terminated if this user signs out. Other secondary users are free to sign in and out. See the GDK documention for details. */
+	UPROPERTY(EditAnywhere, Config, Category=Features)
+	bool bUseSimplifiedUserModel = false;
+
+	/* List of custom protocol schemes to register (March 2023 GDK onwards) */
+	UPROPERTY(EditAnywhere, Config, Category=Features, Meta=(TitleProperty="RegisteredProtocol", ForMSGameStore))
+	TArray<FGDKUriActivationProtocol> UriActivationProtocols;
+
+
+
+	/* Name of the product in Partner Center. Must match Partner Center configuration */
+	UPROPERTY(EditAnywhere, Config, Category=PartnerCenter)
+	FString DefaultDisplayName;
+
+	/* Package/Identity/Name. Must match Partner Center configuration */
+	UPROPERTY(EditAnywhere, Config, Category=PartnerCenter, DisplayName="Package/Identity/Name")
+	FString PackageName;
+
+	/* Package/Identity/Publisher. Must match Partner Center configuration */
+	UPROPERTY(EditAnywhere, Config, Category=PartnerCenter, DisplayName="Package/Identity/Publisher")
+	FString PublisherName;
+
+	/* Package/Properties/PublisherDisplayName. Must match Partner Center configuration */
+	UPROPERTY(EditAnywhere, Config, Category=PartnerCenter, DisplayName="Package/Properties/PublisherDisplayName")
+	FString PublisherDisplayName;
+
+	/* Store Id. Must match Partner Center configuration. */
+	UPROPERTY(EditAnywhere, Config, Category=PartnerCenter, DisplayName="Store ID")
+	FString StoreId;
+
+	/* Xbox Live Title Id. Get this value from your Partner Center configuration. This value is required for title submission. */
+	UPROPERTY(EditAnywhere, Config, Category=PartnerCenter, DisplayName="Xbox Title ID")
+	FString TitleId;
+
+	/* Xbox Live Service Config Id (SCID). Get this value from your Partner Center configuration. This value is required for title submission. */
+	UPROPERTY(EditAnywhere, Config, Category=PartnerCenter, DisplayName="Xbox Services Configuration ID (SCID)")
+	FString PrimaryServiceConfigId;
+
+	/* MSA App Id. Get this value from your Partner Center configuration. This value is required for title submission. */
+	UPROPERTY(EditAnywhere, Config, Category=PartnerCenter, DisplayName="MSA app Id")
+	FString MSAAppId;
+
+	/* Legacy Product GUID. Must match Partner Center configuration. Ignored if StoreId is specified */
+	UPROPERTY(EditAnywhere, Config, Category=PartnerCenter, DisplayName="Legacy Product ID", Meta=(EditCondition="CanEditProductId()"))
+	FString ProductId;
+
+	UFUNCTION()
+	UE_API bool CanEditProductId() const;
+
+
+
+	/* Package version in the form 0.0.0.0 */
+	UPROPERTY(EditAnywhere, Config, Category=PartnerCenter)
+	FString PackageVersion;
+
+	/* Use Include the current engine changelist in the package version number. */
+	UPROPERTY(EditAnywhere, Config, Category=PartnerCenter, AdvancedDisplay)
+	bool bIncludeEngineVersionInPackageVersion;
+
+
+
+
+	/* Define custom chunk mapping */
+	UPROPERTY(EditAnywhere, Config, Category="Chunk Install", AdvancedDisplay, Meta = (TitleProperty="Pattern"))
+	TArray<FGDKCustomChunkMapping> CustomChunkMapping;
+
+	/* Whether to use simple On-Demand intelligent install chunks, or new Features and Recipes for defining chunks and installations */
+	UPROPERTY(EditAnywhere, Config, Category="Chunk Install")
+	bool bUseFeaturesAndRecipes;
+
+	/* List of chunks. See the GDK documentation for details */
+	UPROPERTY(EditAnywhere, Config, Category="Chunk Install", Meta = (TitleProperty="ChunkId"))
+	TArray<FGDKIntelligentDeliveryChunk> IntelligentDeliveryChunks;
+
+	/* List of features. See the GDK documentation for details */
+	UPROPERTY(EditAnywhere, Config, Category="Chunk Install", Meta = (EditCondition="bUseFeaturesAndRecipes", TitleProperty="Id"))
+	TArray<FGDKIntelligentDeliveryFeature> IntelligentDeliveryFeatures;
+
+	/* List of recipes. See the GDK documentation for details */
+	UPROPERTY(EditAnywhere, Config, Category="Chunk Install", Meta = (EditCondition="bUseFeaturesAndRecipes", TitleProperty="Id"))
+	TArray<FGDKIntelligentDeliveryRecipe> IntelligentDeliveryRecipes;
+
+
+	// UI helper functions to get the list of known devices, used by FGDKIntelligentDeliveryRecipe & FGDKIntelligentDeliveryChunk
+	UFUNCTION()
+	UE_API virtual TArray<FString> GetKnownDeviceTypes();
+
+	// UI helper function to get the list of known features, used by FGDKIntelligentDeliveryRecipe
+	UFUNCTION()
+	UE_API TArray<FString> GetKnownIntelligentDeliveryFeatures();
+
+	// UI helper function to get the list of known culture StageIDs, used by FGDKIntelligentDeliveryChunk
+	UFUNCTION()
+	UE_API TArray<FString> GetKnownCultureStageIds();
+
+	// UI helper function to get the list of known DLC plugins
+	UFUNCTION()
+	UE_API TArray<FString> GetKnownDLCNames();
+	
+	/* Configuration for DLC. */
+	UPROPERTY(EditAnywhere, Config, Category=DLC, Meta=(TitleProperty="DLCName", DisplayName="DLC Plugins (EXPERIMENTAL)"))
+	TArray<FGDKDLCConfiguration> DLCPlugins;
+};
+
+#undef UE_API
diff --git a/Engine/Source/Developer/Microsoft/GDK/GDKPlatformEditor/Public/GDKTargetSettingsDetails.h b/Engine/Source/Developer/Microsoft/GDK/GDKPlatformEditor/Public/GDKTargetSettingsDetails.h
new file mode 100644
--- /dev/null
+++ b/Engine/Source/Developer/Microsoft/GDK/GDKPlatformEditor/Public/GDKTargetSettingsDetails.h
@@ -0,0 +1,73 @@
+// Copyright Epic Games, Inc. All Rights Reserved.
+
+#pragma once
+#include "PropertyHandle.h"
+#include "IDetailCustomization.h"
+#include "IPropertyTypeCustomization.h"
+#include "GDKPlatformEditorModule.h"
+#include "Templates/SharedPointer.h"
+#include "Widgets/SWindow.h"
+
+#define UE_API GDKPLATFORMEDITOR_API
+
+/**
+* Detail customization for GDK target settings panel
+*/
+class FGDKTargetSettingsDetails : public IDetailCustomization
+{
+public:
+	UE_API FGDKTargetSettingsDetails(const TCHAR* InPlatformName);
+
+	/** IDetailCustomization interface */
+	UE_API virtual void CustomizeDetails(IDetailLayoutBuilder& DetailBuilder) override;
+
+protected:
+	UE_API void NotifyPropertyChanged(FProperty* Property, UObject* OwnerObject, const TSharedRef<class IPropertyUtilities> PropertyUtilities) const;
+	UE_API void NotifyPropertyChanged(FName PropertyName, UObject* OwnerObject, const TSharedRef<class IPropertyUtilities> PropertyUtilities) const;
+
+	virtual void AddAdditionalPlatformItems( TSharedPtr<class SHorizontalBox> PartnerCenterBox, TWeakObjectPtr<class UGDKTargetSettings> GDKTargetSettings ) {};
+	virtual void AddAdditionalGettingStartedItems( TSharedPtr<class SHorizontalBox> GettingStartedBox, TWeakObjectPtr<class UGDKTargetSettings> GDKTargetSettings, const TSharedRef<class IPropertyUtilities> PropertyUtilities ) {};
+
+	UE_API bool IsGDKConfigured(class UGDKTargetSettings* GDKTargetSettings) const;
+
+	FString PlatformName;
+
+private:
+	UE_API void AddGettingStartedItems(TWeakObjectPtr<class UGDKTargetSettings> GDKTargetSettings, TSharedPtr<class SHorizontalBox> PartnerCenterBox, const TSharedRef<class IPropertyUtilities> PropertyUtilities);
+	UE_API void OnClickUpdateFromPartnerCenter( UGDKTargetSettings* GDKTargetSettings, const TSharedRef<class IPropertyUtilities> PropertyUtilities );
+	UE_API void UpdatePartnerCenterData( class UGDKTargetSettings* GDKTargetSettings, const IGDKPlatformEditorModule::FPartnerCenterProduct& Product, const TSharedRef<class IPropertyUtilities> PropertyUtilities );
+
+	TSharedPtr<SWindow> PartnerCenterProductWindow;
+};
+
+class FGDKCultureResourceDetails : public IPropertyTypeCustomization
+{
+public:
+	static TSharedRef<IPropertyTypeCustomization> MakeInstance()
+	{
+		return MakeShareable(new FGDKCultureResourceDetails);
+	}
+
+	/** Overridden to do nothing */
+	virtual void CustomizeHeader(TSharedRef<class IPropertyHandle> StructPropertyHandle, class FDetailWidgetRow& HeaderRow, IPropertyTypeCustomizationUtils& StructCustomizationUtils) override {}
+
+	/** Overridden display image customization options for packaging */
+	virtual void CustomizeChildren(TSharedRef<IPropertyHandle> InStructPropertyHandle, class IDetailChildrenBuilder& ChildBuilder, IPropertyTypeCustomizationUtils& StructCustomizationUtils) override;
+
+protected:
+	/** Displays a dialog to set up the images for the given culture */
+	void OnClickConfigureImages(TSharedRef<IPropertyHandle> InStructPropertyHandle);
+
+	/** Creates a packaging image picker */
+	void AddImagePickerRow(TSharedPtr<class SScrollBox> ImagesBox, const FText& ImageDescription, const FString& DefaultImagePath, const FString& TargetImagePath, int32 ImageWidth, int32 ImageHeight);
+
+private:
+	TSharedPtr<IPropertyHandle> GetDLCNameProperty( TSharedPtr<IPropertyHandle> InStructPropertyHandle ) const;
+	/** Delegate handler to get the path to start picking from */
+	FString GetPickerPath();
+
+	/** Delegate handler to set the path to start picking from */
+	bool HandlePostExternalIconCopy(const FString& InChosenImage);
+};
+
+#undef UE_API
diff --git a/Engine/Source/Runtime/Microsoft/GDK/GDKPackageManifest/GDKPackageManifest.Build.cs b/Engine/Source/Runtime/Microsoft/GDK/GDKPackageManifest/GDKPackageManifest.Build.cs
new file mode 100644
--- /dev/null
+++ b/Engine/Source/Runtime/Microsoft/GDK/GDKPackageManifest/GDKPackageManifest.Build.cs
@@ -0,0 +1,23 @@
+// Copyright Epic Games, Inc. All Rights Reserved.
+
+using UnrealBuildTool;
+
+public class GDKPackageManifest : ModuleRules
+{
+	public GDKPackageManifest( ReadOnlyTargetRules Target) : base(Target)
+	{
+		if (Target.Platform == UnrealTargetPlatform.Win64)
+		{
+			BinariesSubFolder = "GDK";
+		}
+
+		PrivateDependencyModuleNames.AddRange(
+			new string[]
+			{
+				"Core",
+				"Json",
+				"GRDK",
+			}
+		);
+	}
+}
diff --git a/Engine/Source/Runtime/Microsoft/GDK/GDKPackageManifest/Private/GDKPackageManifestModule.cpp b/Engine/Source/Runtime/Microsoft/GDK/GDKPackageManifest/Private/GDKPackageManifestModule.cpp
new file mode 100644
--- /dev/null
+++ b/Engine/Source/Runtime/Microsoft/GDK/GDKPackageManifest/Private/GDKPackageManifestModule.cpp
@@ -0,0 +1,368 @@
+// Copyright Epic Games, Inc. All Rights Reserved.
+
+#include "IGDKPackageManifestModule.h"
+#include "Misc/CoreDelegates.h"
+#include "Misc/Paths.h"
+#include "Misc/FileHelper.h"
+#include "HAL/PlatformMisc.h"
+#include "Serialization/JsonSerializerMacros.h"
+
+#if WITH_GRDK
+THIRD_PARTY_INCLUDES_START
+#include "Microsoft/AllowMicrosoftPlatformTypes.h"
+#include <XPackage.h>
+#include "Microsoft/HideMicrosoftPlatformTypes.h"
+THIRD_PARTY_INCLUDES_END
+#endif //WITH_GRDK
+
+
+
+
+#if WITH_GRDK
+struct FJsonPackageManifestChunkGDK : FJsonSerializable
+{
+	BEGIN_JSON_SERIALIZER
+		JSON_SERIALIZE("id", Data.UnrealChunkID);
+		JSON_SERIALIZE("Tag", Data.Tag);
+		JSON_SERIALIZE("Languages", Data.Languages);
+		JSON_SERIALIZE("IsInitial", Data.bIsInitial);
+		JSON_SERIALIZE_ARRAY("files", Files);
+	END_JSON_SERIALIZER
+
+	TArray<FString> Files;
+	FGDKPackageManifestChunk Data;
+};
+
+struct FJsonPackageManifestFeatureGDK : FJsonSerializable
+{
+	BEGIN_JSON_SERIALIZER
+		JSON_SERIALIZE("id", Data.Id);
+		JSON_SERIALIZE("Tags", Data.Tags);
+	END_JSON_SERIALIZER
+
+	FGDKPackageManifestFeature Data;
+};
+
+struct FJsonPackageManifestGDK: FJsonSerializable
+{
+	BEGIN_JSON_SERIALIZER
+		JSON_SERIALIZE_ARRAY_SERIALIZABLE("chunks", Chunks, FJsonPackageManifestChunkGDK);
+		JSON_SERIALIZE_ARRAY_SERIALIZABLE("features", Features, FJsonPackageManifestFeatureGDK);
+	END_JSON_SERIALIZER
+
+	TArray<FJsonPackageManifestChunkGDK> Chunks;
+	TArray<FJsonPackageManifestFeatureGDK> Features;
+};
+#endif //WITH_GRDK
+
+
+
+class FGDKPackageManifestModule : public IGDKPackageManifestModule
+{
+public:
+	virtual void StartupModule() override
+	{
+#if WITH_GRDK
+		bIsPackagedProcess = XPackageIsPackagedProcess();
+
+		// cache package identifier
+		FMemory::Memzero(GamePackageInfo.PackageIdentifier);
+		XPackageGetCurrentProcessPackageIdentifier(XPACKAGE_IDENTIFIER_MAX_LENGTH, GamePackageInfo.PackageIdentifier);
+
+		// load the packaging manifest
+		FString PackageManifestText = FPlatformMisc::LoadTextFileFromPlatformPackage(TEXT("package.manifest"));
+		ParseManifest(PackageManifestText, GamePackageInfo );
+#else
+		bIsPackagedProcess = false;
+#endif
+	}
+
+
+	virtual void ShutdownModule() override
+	{
+	}
+
+	virtual bool HasManifest() const override
+	{
+#if WITH_GRDK
+		return GamePackageInfo.bManifestLoaded;
+#else
+		return false;
+#endif
+	}
+
+
+	virtual bool IsPakFileInstalled(const FString& InFilename) const override
+	{
+		return IsPakFileInstalled(nullptr, InFilename);
+	}
+	
+	virtual bool IsPakFileInstalled(XPackageMountHandle MountHandle, const FString& InFilename) const override
+	{
+#if WITH_GRDK
+		// pak files will always be installed in a non-packaged build
+		if (!bIsPackagedProcess)
+		{
+			return true;
+		}
+
+		const FPackageManifestInfo& ManifestInfo = GetManifestInfo(MountHandle);
+
+		// lookup the chunk id for this pak file
+		const FString CleanFilename = FPaths::GetCleanFilename(InFilename).ToLower();
+		const int32* UnrealChunkIDPtr = ManifestInfo.PakToChunkMap.Find(CleanFilename);
+		if (UnrealChunkIDPtr == nullptr || (*UnrealChunkIDPtr) < 0)
+		{
+			return true;
+		}
+
+		// build chunk selector
+		XPackageChunkSelector ChunkSelector;
+		ChunkSelector.type = XPackageChunkSelectorType::Chunk;
+		ChunkSelector.chunkId = (*UnrealChunkIDPtr) + ChunkIDOffset;
+
+		// get chunk availability
+		XPackageChunkAvailability ChunkAvailability = XPackageChunkAvailability::Unavailable;
+		HRESULT hResult = XPackageFindChunkAvailability(ManifestInfo.PackageIdentifier, 1, &ChunkSelector, &ChunkAvailability);
+		UE_CLOG(FAILED(hResult), LogInit, Error, TEXT("Failed to query pak chunk %d availability for %s. 0x%X"), *UnrealChunkIDPtr, *CleanFilename, hResult);
+
+		return SUCCEEDED(hResult) && (ChunkAvailability == XPackageChunkAvailability::Ready);
+#else
+		return true;
+#endif
+	}
+
+	virtual int32 GetChunkIDFromPakchunkIndex(int32 PakchunkIndex) const override
+	{
+#if WITH_GRDK
+		for (const TPair<FString, int32>& PakToChunkPair : GamePackageInfo.PakToChunkMap)
+		{
+			int32 FilePakchunkIndex = FPlatformMisc::GetPakchunkIndexFromPakFile(PakToChunkPair.Key);
+			if (FilePakchunkIndex == PakchunkIndex)
+			{
+				return PakToChunkPair.Value;
+			}
+		}
+#endif
+
+		// unknown - fall back to assuming 1:1
+		return int32(PakchunkIndex);
+	}
+
+	virtual TArray<FString> GetPakFilesInChunk(int32 ChunkID) const override
+	{
+		TArray<FString> Result;
+#if WITH_GRDK
+		const TArray<FString>* PakFilesInChunkPtr = GamePackageInfo.ChunkToPakPaths.Find(ChunkID);
+		if (PakFilesInChunkPtr)
+		{
+			Result = (*PakFilesInChunkPtr);
+		}
+#endif
+		return MoveTemp(Result);
+	}
+
+
+	virtual const TArray<FGDKPackageManifestChunk>& GetChunks() const override
+	{
+		return GamePackageInfo.Chunks;
+	}
+
+	virtual const FGDKPackageManifestChunk* GetChunk(int32 UnrealChunkID) const override
+	{
+		return GamePackageInfo.Chunks.FindByPredicate([UnrealChunkID](const FGDKPackageManifestChunk& Chunk)
+		{
+			return UnrealChunkID == Chunk.UnrealChunkID;
+		});
+	}
+
+	virtual const TArray<FGDKPackageManifestFeature>& GetFeatures() const override
+	{
+		return GamePackageInfo.Features;
+	}
+
+	virtual const TArray<FName>& GetLanguages() const override
+	{
+		return GamePackageInfo.Languages;
+	}
+
+	virtual const TArray<FName>& GetTags() const override
+	{
+		return GamePackageInfo.Tags;
+	}
+
+	virtual bool GetUnrealChunkIDsByFeature( const FString& Feature, TArray<int32>& OutUnrealChunkIDs ) const override
+	{
+
+#if WITH_GRDK
+		const TArray<int32>* ResultPtr = GamePackageInfo.FeatureToUnrealChunkIDMap.Find(Feature);
+		if (ResultPtr)
+		{
+			OutUnrealChunkIDs = *ResultPtr;
+			return true;
+		}
+#endif
+		return false;
+	}
+
+
+	// DLC support
+	virtual bool LoadManifestFromDLC( XPackageMountHandle MountHandle, const char* PackageIdentifier, const FString& MountPath ) override
+	{
+#if WITH_GRDK
+		FPackageManifestInfo& ManifestInfo = DLCPackageInfo.Add(MountHandle);
+		FCStringAnsi::Strncpy(ManifestInfo.PackageIdentifier, PackageIdentifier, XPACKAGE_IDENTIFIER_MAX_LENGTH);
+
+		FString PackageManifestText;
+		if (FFileHelper::LoadFileToString(PackageManifestText, &IPlatformFile::GetPlatformPhysical(), *FPaths::Combine( MountPath, TEXT("package.manifest"))))
+		{
+			ParseManifest(PackageManifestText, ManifestInfo );
+		}
+
+		return true;
+#else
+		return false;
+#endif
+	}
+
+	virtual void UnloadDLC( XPackageMountHandle MountHandle ) override
+	{
+#if WITH_GRDK
+		DLCPackageInfo.Remove(MountHandle);
+#endif
+	}
+
+
+private:
+
+	/**
+	 * GDK chunk numbering is off by 2 to account for the registration chunk and binary chunk. To prevent double counting, this is only
+	 * applied when sending chunk ids to the system, not for our internal API or data.
+	 */
+	static const int ChunkIDOffset = 2;
+
+	struct FPackageManifestInfo
+	{
+#if WITH_GRDK
+		char PackageIdentifier[XPACKAGE_IDENTIFIER_MAX_LENGTH];
+		TMap<FString, int32> PakToChunkMap;
+		TMap<FString, TArray<int32>> FeatureToUnrealChunkIDMap;
+		TMap<int32, TArray<FString>> ChunkToPakPaths;
+
+		bool bManifestLoaded;
+#endif
+		TArray<FGDKPackageManifestChunk> Chunks;
+		TArray<FGDKPackageManifestFeature> Features;
+		TArray<FName> Languages;
+		TArray<FName> Tags;
+	};
+
+	FPackageManifestInfo GamePackageInfo;
+#if WITH_GRDK
+	TMap<XPackageMountHandle, FPackageManifestInfo> DLCPackageInfo;
+#endif
+
+
+	bool bIsPackagedProcess;
+
+
+
+#if WITH_GRDK
+	bool ParseManifest( const FString& PackageManifestText, FPackageManifestInfo& ManifestInfo )
+	{
+		TUniquePtr<struct FJsonPackageManifestGDK> Manifest = MakeUnique<FJsonPackageManifestGDK>();
+		ManifestInfo.bManifestLoaded = Manifest->FromJson(PackageManifestText);
+
+		if (ManifestInfo.bManifestLoaded)
+		{
+			TMap<FString,TArray<int32>> ChunkTagToUnrealChunkIDMap;
+
+			for (const FJsonPackageManifestChunkGDK& Chunk : Manifest->Chunks)
+			{
+				// build pak file lookup table
+				for (const FString& Filename : Chunk.Files)
+				{
+					if (Filename.EndsWith(TEXT(".pak")) || Filename.EndsWith(TEXT(".uondemandtoc")))
+					{
+						FString CleanFilename = FPaths::GetCleanFilename(Filename).ToLower();
+						ManifestInfo.PakToChunkMap.Add(MoveTemp(CleanFilename), Chunk.Data.UnrealChunkID);
+						ManifestInfo.ChunkToPakPaths.FindOrAdd(Chunk.Data.UnrealChunkID).Add(Filename);
+					}
+				}
+
+				// cache chunk data
+				ManifestInfo.Chunks.Add(Chunk.Data);
+
+				// cache chunk tag
+				TArray<FString> TagArray;
+				Chunk.Data.Tag.ParseIntoArray(TagArray, TEXT(";") );
+				for (FString Tag : TagArray)
+				{
+					FString Tag2;
+					if (Tag.Split(FString(TEXT("#")), &Tag, &Tag2 ))
+					{
+						ChunkTagToUnrealChunkIDMap.FindOrAdd(Tag2).AddUnique(Chunk.Data.UnrealChunkID);
+						ManifestInfo.Tags.AddUnique(FName(Tag2));
+					}
+					ChunkTagToUnrealChunkIDMap.FindOrAdd(Tag).AddUnique(Chunk.Data.UnrealChunkID);
+					ManifestInfo.Tags.AddUnique(FName(Tag));
+				}
+
+				// cache languages
+				TArray<FString> LanguageArray;
+				Chunk.Data.Languages.ParseIntoArray(LanguageArray, TEXT(";") );
+				for (const FString& Language : LanguageArray)
+				{
+					ManifestInfo.Languages.AddUnique(FName(Language));
+				}
+			}
+
+			for (const FJsonPackageManifestFeatureGDK& Feature : Manifest->Features)
+			{
+				// cache feature data
+				ManifestInfo.Features.Add(Feature.Data);
+
+				// cache feature tags
+				TArray<FString> TagArray;
+				Feature.Data.Tags.ParseIntoArray(TagArray, TEXT(";") );
+				for (const FString& Tag : TagArray)
+				{
+					if (const TArray<int32>* UnrealChunkIDsPtr = ChunkTagToUnrealChunkIDMap.Find(Tag))
+					{
+						TArray<int32>& FeatureChunkIDs = ManifestInfo.FeatureToUnrealChunkIDMap.FindOrAdd(Feature.Data.Id);
+						for (int32 UnrealChunkID : *UnrealChunkIDsPtr)
+						{
+							FeatureChunkIDs.AddUnique(UnrealChunkID);
+						}
+					}
+				}
+			}
+		}
+
+		return ManifestInfo.bManifestLoaded;
+	}
+
+
+	const FPackageManifestInfo& GetManifestInfo( XPackageMountHandle MountHandle ) const
+	{
+		if (MountHandle == nullptr)
+		{
+			return GamePackageInfo;
+		}
+		else
+		{
+			const FPackageManifestInfo* ManifestInfo = DLCPackageInfo.Find(MountHandle);
+			checkf(ManifestInfo, TEXT("invalid package mount handle"));
+
+			return *ManifestInfo;
+		}
+	}
+#endif //WITH_GRDK
+
+};
+
+
+
+
+IMPLEMENT_MODULE(FGDKPackageManifestModule, GDKPackageManifest);
+
diff --git a/Engine/Source/Runtime/Microsoft/GDK/GDKPackageManifest/Public/IGDKPackageManifestModule.h b/Engine/Source/Runtime/Microsoft/GDK/GDKPackageManifest/Public/IGDKPackageManifestModule.h
new file mode 100644
--- /dev/null
+++ b/Engine/Source/Runtime/Microsoft/GDK/GDKPackageManifest/Public/IGDKPackageManifestModule.h
@@ -0,0 +1,76 @@
+// Copyright Epic Games, Inc. All Rights Reserved.
+
+#pragma once
+
+#include "CoreMinimal.h"
+#include "GenericPlatform/GenericPlatformChunkInstall.h"
+#include "Modules/ModuleManager.h"
+
+struct FGDKPackageManifestChunk
+{
+	int32 UnrealChunkID = 0;
+	FString Tag; // raw tag string with ; and # separators
+	FString Languages;
+	bool bIsInitial;
+};
+
+struct FGDKPackageManifestFeature
+{
+	FString Id;
+	FString Tags; // raw tag string with ; separators
+};
+
+typedef struct XPackageMountInstance* XPackageMountHandle;
+
+class IGDKPackageManifestModule : public IModuleInterface, public IPlatformChunkInstallManifest
+{
+public:
+	/**
+	 * Singleton-like access to this module's interface.  This is just for convenience!
+	 * Beware of calling this during the shutdown phase, though.  Your module might have been unloaded already.
+	 *
+	 * @return Returns singleton instance, loading the module on demand if needed
+	 */
+	static inline IGDKPackageManifestModule& Get()
+	{
+		return FModuleManager::LoadModuleChecked<IGDKPackageManifestModule>("GDKPackageManifest");
+	}
+
+	// determine whether the given pak file is installed or just a download placeholder with no file system backing
+	virtual bool IsPakFileInstalled(const FString& InFilename) const = 0;
+
+	// determine if the package manifest was loaded successfully
+	virtual bool HasManifest() const override = 0;
+
+	// determine the Unreal chunk ID that the given pak chunk index is stored in
+	virtual int32 GetChunkIDFromPakchunkIndex(int32 PakchunkIndex) const override = 0;
+
+	// get the pak files that are stored in the given Unreal chunk
+	virtual TArray<FString> GetPakFilesInChunk(int32 ChunkID) const override= 0;
+
+
+	// get all chunk data
+	virtual const TArray<FGDKPackageManifestChunk>& GetChunks() const = 0;
+
+	// get data for a specific unreal chunk
+	virtual const FGDKPackageManifestChunk* GetChunk(int32 UnrealChunkID) const = 0;
+
+	// get all feature data
+	virtual const TArray<FGDKPackageManifestFeature>& GetFeatures() const = 0;
+
+	// get all languages
+	virtual const TArray<FName>& GetLanguages() const = 0;
+
+	// get all tags
+	virtual const TArray<FName>& GetTags() const = 0;
+
+
+	// get chunks that are referred to by the given feature
+	virtual bool GetUnrealChunkIDsByFeature( const FString& Feature, TArray<int32>& OutUnrealChunkIDs ) const = 0;
+
+	// DLC support
+	virtual bool LoadManifestFromDLC( XPackageMountHandle MountHandle, const char* PackageIdentifier, const FString& MountPath ) = 0;
+	virtual void UnloadDLC( XPackageMountHandle MountHandle ) = 0;
+	virtual bool IsPakFileInstalled(XPackageMountHandle MountHandle, const FString& InFilename) const = 0;
+
+};
diff --git a/Engine/Source/Runtime/Microsoft/GDK/GDKRuntime/GDKRuntime.Build.cs b/Engine/Source/Runtime/Microsoft/GDK/GDKRuntime/GDKRuntime.Build.cs
new file mode 100644
--- /dev/null
+++ b/Engine/Source/Runtime/Microsoft/GDK/GDKRuntime/GDKRuntime.Build.cs
@@ -0,0 +1,15 @@
+// Copyright Epic Games, Inc. All Rights Reserved.
+
+using UnrealBuildTool;
+using System.IO;
+
+public class GDKRuntime : ModuleRules
+{
+	public GDKRuntime(ReadOnlyTargetRules Target) : base(Target)
+	{
+		BinariesSubFolder = "GDK";
+
+		PrivateDependencyModuleNames.Add("Core");
+		PrivateDependencyModuleNames.Add("GRDK");
+	}
+}
diff --git a/Engine/Source/Runtime/Microsoft/GDK/GDKRuntime/Private/GDKHandleTracker.cpp b/Engine/Source/Runtime/Microsoft/GDK/GDKRuntime/Private/GDKHandleTracker.cpp
new file mode 100644
--- /dev/null
+++ b/Engine/Source/Runtime/Microsoft/GDK/GDKRuntime/Private/GDKHandleTracker.cpp
@@ -0,0 +1,228 @@
+// Copyright Epic Games, Inc. All Rights Reserved.
+
+#include "GDKHandleTracker.h"
+
+#if WITH_GDK_HANDLE_TRACKER
+
+#include "CoreMinimal.h"
+#include "HAL/PlatformStackWalk.h"
+#include "HAL/PlatformMisc.h"
+#include "HAL/IConsoleManager.h"
+#include "Misc/ScopeLock.h"
+#include "Misc/Crc.h"
+#include "Containers/Map.h"
+
+THIRD_PARTY_INCLUDES_START
+	#include "Microsoft/AllowMicrosoftPlatformTypes.h"
+	#include <atomic>
+	#include <XSystem.h>
+	#include "Microsoft/HideMicrosoftPlatformTypes.h"
+THIRD_PARTY_INCLUDES_END
+
+
+static FAutoConsoleCommand CmdGDKTrackStart(
+	TEXT("GDK.XSystemHandle.TrackStart"),
+	TEXT("start tracking GDK handles"),
+	FConsoleCommandDelegate::CreateStatic( &FGDKHandleTracker::Start )
+);
+
+static FAutoConsoleCommand CmdGDKTrackStop(
+	TEXT("GDK.XSystemHandle.TrackStop"),
+	TEXT("stop tracking GDK handles & dump any opened handles (not necessary a leak)"),
+	FConsoleCommandDelegate::CreateStatic( &FGDKHandleTracker::Stop )
+);
+
+
+
+
+namespace GDKHandleTrackerInternal
+{
+	static std::atomic<int32> ActiveCounter = {0};
+
+	// TODO: should these be data-driven? should we have a bit array or something?
+	static bool bLimitType = false;
+	static XSystemHandleType LimitType = XSystemHandleType::User;
+
+	// callstack tracking data
+	struct FGDKTrackedCallstack
+	{
+		TArray<uint64> Backtrace;
+		uint32 RefCount;
+		XSystemHandleType HandleType;
+	};
+
+	FCriticalSection TrackLock;
+	TMap<uint32,FGDKTrackedCallstack> TrackedCallstacks; // hash to callstack
+	TMap<XSystemHandle,uint32> TrackedHandles;           // handle to hash
+
+	// high watermark
+	int32 MaxTrackedHandles = 0;
+	int32 MaxTrackedCallstacks = 0;
+
+
+	const TCHAR* GetHandleTypeName( XSystemHandleType HandleType )
+	{
+		switch(HandleType)
+		{
+			case XSystemHandleType::AppCaptureScreenshotStream: return TEXT("AppCaptureScreenshotStream");
+			case XSystemHandleType::DisplayTimeoutDeferral    : return TEXT("DisplayTimeoutDeferral");
+			case XSystemHandleType::GameSaveContainer         : return TEXT("GameSaveContainer");
+			case XSystemHandleType::GameSaveProvider          : return TEXT("GameSaveProvider");
+			case XSystemHandleType::GameSaveUpdate            : return TEXT("GameSaveUpdate");
+			case XSystemHandleType::PackageInstallationMonitor: return TEXT("PackageInstallationMonitor");
+			case XSystemHandleType::PackageMount              : return TEXT("PackageMount");
+			case XSystemHandleType::SpeechSynthesizer         : return TEXT("SpeechSynthesizer");
+			case XSystemHandleType::SpeechSynthesizerStream   : return TEXT("SpeechSynthesizerStream");
+			case XSystemHandleType::StoreContext              : return TEXT("StoreContext");
+			case XSystemHandleType::StoreLicense              : return TEXT("StoreLicense");
+			case XSystemHandleType::StoreProductQuery         : return TEXT("StoreProductQuery");
+			case XSystemHandleType::TaskQueue                 : return TEXT("TaskQueue");
+			case XSystemHandleType::User                      : return TEXT("User");
+			case XSystemHandleType::UserSignOutDeferral       : return TEXT("UserSignOutDeferral");
+			default: return TEXT("(unknown)");
+		}
+	}
+
+	uint32 TrackCallstack( XSystemHandleType HandleType )
+	{
+		// capture callstack
+		static const uint32 MaxBacktraceDepth = 32;
+		uint64 Backtrace[MaxBacktraceDepth];
+		uint32 BacktraceDepth = FPlatformStackWalk::CaptureStackBackTrace( Backtrace, MaxBacktraceDepth );
+
+		// skip the actual capture & tracking items
+		static const uint32 BacktraceEntriesToIgnore = 4; // xgameruntime handle tracker -> GDKSystemHandleCallback() -> TrackCallstack() -> FPlatformStackWalk::CaptureStackBackTrace()
+		uint64* BacktraceStart = Backtrace + BacktraceEntriesToIgnore;
+		BacktraceDepth -= BacktraceEntriesToIgnore;
+
+		// hash the callstack and see if we have it already
+		uint32 CallstackHash = FCrc::MemCrc32(BacktraceStart, sizeof(uint64) * BacktraceDepth, 0);
+		FGDKTrackedCallstack* Existing = TrackedCallstacks.Find(CallstackHash);
+		if (Existing == nullptr)
+		{
+			Existing = &TrackedCallstacks.Add( CallstackHash );
+			Existing->Backtrace.Append(BacktraceStart, BacktraceDepth);
+			Existing->RefCount = 0;
+			Existing->HandleType = HandleType;
+		}
+		else
+		{
+			// shouldn't ever get two handles from the same callstack (if we do, they're being created from within one GDK handle creation function which is unlikely. Would need to roll the handle type into the hash in this case)
+			check(Existing->HandleType == HandleType);
+		}
+
+		// add reference & return
+		Existing->RefCount++;
+		return CallstackHash;
+	}
+
+
+	void FreeCallstack(uint32 CallstackHash)
+	{
+		// lookup the callstack and clear it if this is the last reference
+		FGDKTrackedCallstack* Existing = TrackedCallstacks.Find(CallstackHash);
+		if (ensure(Existing != nullptr))
+		{
+			Existing->RefCount--;
+			if (Existing->RefCount <= 0)
+			{
+				TrackedCallstacks.Remove(CallstackHash);
+			}
+		}
+	}
+
+
+	static void CALLBACK GDKSystemHandleCallback( XSystemHandle Handle, XSystemHandleType HandleType, XSystemHandleCallbackReason Reason, void* Context )
+	{
+		// filter by type
+		if (bLimitType && LimitType != HandleType)
+		{
+			return;
+		}
+
+		if (Reason == XSystemHandleCallbackReason::Created)
+		{
+			// track this handle's callstack & update high watermark
+			FScopeLock Lock(&TrackLock);
+			TrackedHandles.Add(Handle, TrackCallstack(HandleType) );
+
+			MaxTrackedCallstacks = FMath::Max( MaxTrackedCallstacks, TrackedCallstacks.Num() );
+			MaxTrackedHandles = FMath::Max( MaxTrackedHandles, TrackedHandles.Num() );
+		}
+		else if (Reason == XSystemHandleCallbackReason::Destroyed)
+		{
+			// release this handle's callstack
+			FScopeLock Lock(&TrackLock);
+			const uint32* CallstackPtr = TrackedHandles.Find(Handle);
+			if (CallstackPtr != nullptr) // NB. we may not find it if tracking has already been dumped & restarted
+			{
+				FreeCallstack(*CallstackPtr);
+				TrackedHandles.Remove(Handle);
+			}
+		}
+	}
+
+	// note: uses LowLevelOutputDebugString not UE_LOG because it's typically called after engine shutdown & UE_LOG is no longer available
+	static void DumpHandles()
+	{
+		FScopeLock Lock(&TrackLock);
+
+		// dump callstacks & handle counters
+		if (TrackedHandles.Num() > 0 )
+		{
+			FPlatformMisc::LowLevelOutputDebugString( TEXT("\n\n"));
+
+			for ( const TPair<uint32,FGDKTrackedCallstack>& TrackedCallstack : TrackedCallstacks )
+			{
+				FPlatformMisc::LowLevelOutputDebugStringf( TEXT("%d active %s GDK handles detected from callstack:\n"), TrackedCallstack.Value.RefCount, GetHandleTypeName(TrackedCallstack.Value.HandleType) );
+
+				// dump callstack
+				for (int32 Depth = 0; Depth < TrackedCallstack.Value.Backtrace.Num(); Depth++)
+				{
+					ANSICHAR Buffer[1024];
+					Buffer[0] = '\0';
+					FPlatformStackWalk::ProgramCounterToHumanReadableString(Depth, TrackedCallstack.Value.Backtrace[Depth], Buffer, sizeof(Buffer) );
+					FPlatformMisc::LowLevelOutputDebugStringf( TEXT("\t%s\n"), ANSI_TO_TCHAR(Buffer));
+				}
+				FPlatformMisc::LowLevelOutputDebugString( TEXT("\n\n"));
+			}
+		}
+
+		// dump final report & reset trackers
+		FPlatformMisc::LowLevelOutputDebugStringf( TEXT("Total %d handles, %d unique callstacks. (Maximum was %d handles, %d unique callstacks)\n"), TrackedHandles.Num(), TrackedCallstacks.Num(), MaxTrackedHandles, MaxTrackedCallstacks);
+		MaxTrackedHandles = 0;
+		MaxTrackedCallstacks = 0;
+		TrackedCallstacks.Reset();
+		TrackedHandles.Reset();
+	}
+}
+
+void FGDKHandleTracker::Start()
+{
+	if ((++GDKHandleTrackerInternal::ActiveCounter) == 1)
+	{
+		HRESULT hResult = XSystemHandleTrack( GDKHandleTrackerInternal::GDKSystemHandleCallback, nullptr);
+		if (FAILED(hResult))
+		{
+			FPlatformMisc::LowLevelOutputDebugStringf( TEXT("Failed to start GDK handle tracking : 0x%X\n"), hResult);
+		}
+	}
+}
+
+void FGDKHandleTracker::Stop()
+{
+	if ((--GDKHandleTrackerInternal::ActiveCounter) == 0)
+	{
+		CA_SUPPRESS(6387) // first parameter is tagged _In_ but the documentation says nullptr used to disable handle tracking
+		XSystemHandleTrack(nullptr, nullptr);
+		GDKHandleTrackerInternal::DumpHandles();
+	}
+	else
+	{
+		FPlatformMisc::LowLevelOutputDebugString( TEXT("Other GDK trackers are still active - ignoring stop\n"));
+	}
+}
+
+
+
+#endif //WITH_GDK_HANDLE_TRACKER
diff --git a/Engine/Source/Runtime/Microsoft/GDK/GDKRuntime/Private/GDKRuntimeInitialize.cpp b/Engine/Source/Runtime/Microsoft/GDK/GDKRuntime/Private/GDKRuntimeInitialize.cpp
new file mode 100644
--- /dev/null
+++ b/Engine/Source/Runtime/Microsoft/GDK/GDKRuntime/Private/GDKRuntimeInitialize.cpp
@@ -0,0 +1,337 @@
+// Copyright Epic Games, Inc. All Rights Reserved.
+
+
+#include "GDKRuntimeModule.h"
+#if WITH_GRDK
+#include "GDKHandleTracker.h"
+#include "HAL/PlatformProcess.h" 
+#include "Misc/ConfigCacheIni.h"
+#include "Misc/Paths.h"
+#include "Misc/CommandLine.h"
+#include "Misc/App.h"
+#include "Misc/FileHelper.h"
+#include "Containers/AnsiString.h"
+#include "Microsoft/AllowMicrosoftPlatformTypes.h"
+THIRD_PARTY_INCLUDES_START
+#include <grdk.h>
+#include <XGameRuntimeInit.h>
+#include <XError.h>
+#include <XGameErr.h>
+#include <appmodel.h>
+THIRD_PARTY_INCLUDES_END
+#include "Microsoft/HideMicrosoftPlatformTypes.h"
+
+#ifndef WITH_GRDK_DEV_RUNTIME_INIT
+	#define WITH_GRDK_DEV_RUNTIME_INIT PLATFORM_WINDOWS && !UE_BUILD_SHIPPING && (_GRDK_EDITION >= 240600)
+#endif
+
+#if WITH_GRDK_DEV_RUNTIME_INIT
+static HRESULT InitXGameRuntimeForDevelopment(const TCHAR* ConfigIniSection);
+#endif
+
+
+static int32 GDKEnvironmentInitRefCount = 0;
+static bool bGDKEnvironmentInitialized = false;
+
+static bool bGDKIsTrackingHandles = false;
+
+
+bool SetupGDKEnvironment(const TCHAR* ConfigIniSection)
+{
+	// initialize on first call only
+	if ((++GDKEnvironmentInitRefCount) > 1)
+	{
+		return bGDKEnvironmentInitialized;
+	}
+
+	// GDK runtime is not needed for commandlets
+	if (IsRunningCommandlet())
+	{
+		return false;
+	}
+
+	// initialize the gaming runtime
+#if WITH_GRDK_DEV_RUNTIME_INIT
+	HRESULT hResult = InitXGameRuntimeForDevelopment(ConfigIniSection);
+	if (FAILED(hResult))
+	{
+		hResult = XGameRuntimeInitialize();
+	}
+#else
+	HRESULT hResult = XGameRuntimeInitialize();
+#endif
+
+	// Workaround race condition with Clang where xgameruntime.dll might not be present by this point
+#if defined(__clang__) && !PLATFORM_WINDOWS
+	int32 RetryCount = 0;
+	while (hResult == E_GAMERUNTIME_DLL_NOT_FOUND && RetryCount < 100)
+	{
+		// Sleep 200ms and try again
+		FPlatformProcess::Sleep(0.200f);
+		UE_LOG(LogGDK, Warning, TEXT("Error: Game Runtime dll wasn't found! Retrying"));
+		hResult = XGameRuntimeInitialize();
+		RetryCount++;
+	}
+#endif // __clang__ && !PLATFORM_WINDOWS
+
+	// check for specific errors
+#if PLATFORM_WINDOWS
+	UE_CLOG(hResult == E_GAMERUNTIME_DLL_NOT_FOUND, LogGDK, Error, TEXT("Error: Game Runtime is not installed"));
+	UE_CLOG(hResult == E_GAMERUNTIME_VERSION_MISMATCH, LogGDK, Error, TEXT("Game Runtime version mismatch"));
+	UE_CLOG(hResult==0x80073CFC, LogGDK, Error, TEXT("GamingServicesRuntime may need reinstalling") );
+	UE_CLOG(hResult==0x87e5001f/*E_GAME_MISSING_GAME_CONFIG*/, LogGDK, Error, TEXT("Missing MicrosoftGame.config"));
+#else
+	UE_CLOG(hResult == E_GAMERUNTIME_DLL_NOT_FOUND, LogGDK, Fatal, TEXT("Error: Game Runtime is not installed"));
+	UE_CLOG(hResult == E_GAMERUNTIME_VERSION_MISMATCH, LogGDK, Fatal, TEXT("Game Runtime version mismatch"));
+#endif
+	UE_CLOG(FAILED(hResult), LogGDK, Error, TEXT("XGameRuntimeInitialize failed: 0x%X"), hResult);
+
+	// early out of the runtime failed to initialize
+	bGDKEnvironmentInitialized = SUCCEEDED(hResult);
+	if (!bGDKEnvironmentInitialized)
+	{
+		return false;
+	}
+
+	// check for gdk lifetime command line features
+#if WITH_GDK_HANDLE_TRACKER
+	if (FParse::Param(FCommandLine::Get(),TEXT("gdkhandletrack") ) )
+	{
+		bGDKIsTrackingHandles = true;
+		FGDKHandleTracker::Start();
+	}
+#endif
+
+	return true;
+}
+
+void TeardownGDKEnvironment()
+{
+	if ((--GDKEnvironmentInitRefCount) != 0)
+	{
+		return;
+	}
+
+	if (!bGDKEnvironmentInitialized)
+	{
+		return;
+	}
+
+#if WITH_GDK_HANDLE_TRACKER
+	if (bGDKIsTrackingHandles)
+	{
+		FGDKHandleTracker::Stop();
+		bGDKIsTrackingHandles = false;
+	}
+#endif
+
+	XErrorSetOptions(XErrorOptions::OutputDebugStringOnError, XErrorOptions::OutputDebugStringOnError); // #jira UE-141916 October GDK bug: do not break on error because there are active objects on WinGDK shutdown (even in a new Direct3D 12 Desktop Game project.) 
+	XGameRuntimeUninitialize();
+
+	bGDKEnvironmentInitialized = false;
+}
+
+
+#if WITH_GRDK_DEV_RUNTIME_INIT
+static FAnsiString FindMSGameConfigFile()
+{
+	static const TCHAR* MSGameConfig = TEXT("MicrosoftGame.config");
+	
+	// check project root
+	FString MSGameConfigFile = FPaths::Combine( FPlatformProcess::BaseDir(), TEXT("../../.."), MSGameConfig);
+	if (FPaths::FileExists(MSGameConfigFile))
+	{
+		return *MSGameConfigFile;
+	}
+
+	// check binaries directory
+	MSGameConfigFile = FPaths::Combine( FPlatformProcess::BaseDir(), MSGameConfig);
+	if (FPaths::FileExists(MSGameConfigFile))
+	{
+		return *MSGameConfigFile;
+	}
+
+	// check the saved folder
+#if !WITH_EDITOR
+	// this is primarily to support F5 in Visual Studio without specifying a -basedir ... in this case the game is loading cooked content from the Saved/Cooked directory or via Zen streaming
+	// it will most likely be necessary to specify -plugin=<path to>\Engine\Platforms\GDK\Plugins on the command line while the GDK runtime plugin is under an NDA platform folder.
+	if (FPaths::IsProjectFilePathSet())
+	{
+		MSGameConfigFile = FPaths::Combine( FPaths::ProjectSavedDir(), TEXT("Win64"), TEXT("Manifest"), TEXT("MicrosoftGame.config") );
+		if (FPaths::FileExists(MSGameConfigFile))
+		{
+			return *MSGameConfigFile;
+		}
+	}
+#endif
+
+	return FAnsiString();
+}
+#endif // WITH_GRDK_DEV_RUNTIME_INIT
+
+
+
+
+
+
+#if WITH_GRDK_DEV_RUNTIME_INIT
+static FAnsiString MakeInlineMSGameConfig(const TCHAR* ConfigSectionName )
+{
+	// must have config section (code issue)
+	if (ConfigSectionName == nullptr)
+	{
+		UE_LOG(LogGDK, Warning, TEXT("Config section name has not been specified."));
+		return FAnsiString();
+	}
+
+	// Newer GDKs require these two properties to be specified before the default user can be added. 
+	// While many GDK runtime functions still function without a user, most of the useful functionality is user-focused so just disable GDK runtime if we can't add a user.
+	FString TitleId, MSAAppId;
+	if (!GConfig->GetString(ConfigSectionName, TEXT("TitleId"),  TitleId,  GEngineIni) || TitleId.IsEmpty() ||
+		!GConfig->GetString(ConfigSectionName, TEXT("MSAAppId"), MSAAppId, GEngineIni) || MSAAppId.IsEmpty())
+	{
+		UE_LOG(LogGDK, Warning, TEXT("TitleId and MSAppId must be specified in MSGamingRuntime plugin settings before GDK runtime can be used."));
+		return FAnsiString();
+	}
+
+	// Read package identity
+	FString PackageName = FApp::GetProjectName();
+	GConfig->GetString(ConfigSectionName, TEXT("PackageName"), PackageName, GEngineIni);
+	if (PackageName.IsEmpty())
+	{
+		PackageName = TEXT("DefaultUEProject");
+	}
+
+	FString PublisherName;
+	GConfig->GetString(ConfigSectionName, TEXT("PublisherName"), PublisherName, GEngineIni);
+	if (PublisherName.IsEmpty())
+	{
+		PublisherName = TEXT("CN=NoPublisher");
+	}
+
+
+	// Write a smallest-possible placeholder MicrosoftGame.config file required to initialize the GDK runtime.
+	return FAnsiString::Printf("<?xml version=\"1.0\" encoding=\"utf-8\"?><Game configVersion=\"1\"><Identity Name=\"%S\" Publisher=\"%S\"/><TitleId>%S</TitleId><MSAAppId>%S</MSAAppId></Game>", *PackageName, *PublisherName, *TitleId, *MSAAppId );
+}
+#endif // WITH_GRDK_DEV_RUNTIME_INIT
+
+
+
+#if WITH_GRDK_DEV_RUNTIME_INIT
+static bool IsPackagedProcess()
+{
+	// editor will never be packaged 
+#if WITH_EDITOR
+	return false;
+#else
+
+	// get the kernel
+	HMODULE hModule = ::GetModuleHandleW(TEXT("kernel32.dll"));
+	if (hModule == nullptr)
+	{
+		return false;
+	}
+
+	// look up the GetCurrentPackageId function
+#pragma warning( push )
+#pragma warning( disable: 4191 )	// unsafe conversion from 'type of expression' to 'type required'
+	typedef LONG(WINAPI *GetCurrentPackageIdProc)(UINT32*, BYTE*);
+	GetCurrentPackageIdProc fnGetCurrentPackageId = (GetCurrentPackageIdProc)::GetProcAddress(hModule, "GetCurrentPackageId");
+#pragma warning( pop )
+
+	if (fnGetCurrentPackageId == nullptr)
+	{
+		return false;
+	}
+
+	UINT32 BufferLength = 0;
+	return fnGetCurrentPackageId(&BufferLength, nullptr) != APPMODEL_ERROR_NO_PACKAGE;
+#endif  //WITH_EDITOR
+}
+#endif // WITH_GRDK_DEV_RUNTIME_INIT
+
+
+
+#if WITH_GRDK_DEV_RUNTIME_INIT
+static HRESULT InitXGameRuntimeForDevelopment(const TCHAR* ConfigIniSection)
+{
+	HRESULT hResult = E_FAIL;
+
+	// early out if this is a packaged build because debugging assistance is not required
+	if (IsPackagedProcess())
+	{
+		UE_LOG(LogGDK, Display, TEXT("Debugging packaged build"));
+		return E_FAIL;
+	}
+
+	// check for a MicrosoftGame.config in known locations
+	FAnsiString MSGameConfigFile = FindMSGameConfigFile();
+	if (!MSGameConfigFile.IsEmpty())
+	{
+		XGameRuntimeOptions Options;
+		Options.gameConfigSource = XGameRuntimeGameConfigSource::File;
+		Options.gameConfig = *MSGameConfigFile;
+		hResult = XGameRuntimeInitializeWithOptions(&Options);
+	}
+
+	// make a placeholder minimal MicrosoftGame.config - just enough to initialize
+	if (FAILED(hResult) && hResult != E_GAMERUNTIME_OPTIONS_NOT_SUPPORTED)
+	{
+		FAnsiString MSGameConfig = MakeInlineMSGameConfig(ConfigIniSection);
+		if (!MSGameConfig.IsEmpty())
+		{
+			XGameRuntimeOptions Options;
+			Options.gameConfigSource = XGameRuntimeGameConfigSource::Inline;
+			Options.gameConfig = *MSGameConfig;
+			hResult = XGameRuntimeInitializeWithOptions(&Options);
+#if WITH_EDITOR
+			UE_CLOG(SUCCEEDED(hResult), LogGDK, Display, TEXT("Using minimal XGameRuntime config for PIE - most settings ignored") );
+#else
+UE_CLOG(SUCCEEDED(hResult), LogGDK, Warning, TEXT("Using temporary XGameRuntime config - most settings ignored") );
+#endif
+			checkf(hResult != E_GAMERUNTIME_GAMECONFIG_BAD_FORMAT, TEXT("autogenerated MSGameConfig is malformed - check the code/GDK release notes"));
+		}
+	}
+
+
+	// this error code indicates that there is a packaged version of the build already installed
+	// in this case we must currently create a temporary minimal MicrosoftGame.config file alongside the .exe
+	// to work around a GDK bug
+	if (hResult == E_GAMERUNTIME_OPTIONS_NOT_SUPPORTED)
+	{
+		// give some details that might be helpful
+		UE_LOG(LogGDK, Display, TEXT("NOTE: There is a packaged version of this game installed as well but it is not the one being debugged.")); 
+		UE_LOG(LogGDK, Display, TEXT("To debug that in visual studio, use Debug -> Other Debug Targets -> Debug Installed App Package."));
+		UE_LOG(LogGDK, Display, TEXT("To debug this executable against the installed package's content, add the path to the installed package's binary folder to the command line.  e.g. -basedir=C:\\XboxGames\\[YourApp]\\Content\\[YourGame]\\Binaries\\Win64\\  (use wdapp.exe list /d to discover the install path)")); 
+
+		// delete any stale file
+		FString TempMSGameConfigFile = FPaths::Combine(FPaths::GetPath(FPlatformProcess::ExecutablePath()), TEXT("MicrosoftGame.config"));
+		if (IPlatformFile::GetPlatformPhysical().FileExists(*TempMSGameConfigFile))
+		{
+			IPlatformFile::GetPlatformPhysical().DeleteFile(*TempMSGameConfigFile);
+		}
+
+		// generate a temporary MicrosoftGame.config alongside the exe file
+		FAnsiString MSGameConfig = MakeInlineMSGameConfig(ConfigIniSection);
+		if (FFileHelper::SaveStringToFile(ANSI_TO_TCHAR(*MSGameConfig), *TempMSGameConfigFile, FFileHelper::EEncodingOptions::ForceAnsi))
+		{
+			hResult = XGameRuntimeInitialize();
+			IPlatformFile::GetPlatformPhysical().DeleteFile(*TempMSGameConfigFile);
+		}
+	}
+
+
+	// cannot find gaming config - have to initialize without one. 
+	// this likely means that we can't sign in a user or connect to xsapi etc. because we won't have a package identity
+	if (FAILED(hResult))
+	{
+		hResult = XGameRuntimeInitialize();
+		UE_CLOG(SUCCEEDED(hResult), LogGDK, Warning, TEXT("Likely using XGameRuntime without a MicrosoftGame.config") );
+	}
+
+	return hResult;
+}
+#endif // WITH_GRDK_DEV_RUNTIME_INIT
+
+
+#endif //WITH_GRDK
diff --git a/Engine/Source/Runtime/Microsoft/GDK/GDKRuntime/Private/GDKRuntimeModule.cpp b/Engine/Source/Runtime/Microsoft/GDK/GDKRuntime/Private/GDKRuntimeModule.cpp
new file mode 100644
--- /dev/null
+++ b/Engine/Source/Runtime/Microsoft/GDK/GDKRuntime/Private/GDKRuntimeModule.cpp
@@ -0,0 +1,1237 @@
+// Copyright Epic Games, Inc. All Rights Reserved.
+
+
+#include "GDKRuntimeModule.h"
+#if WITH_GRDK
+#include "GDKHandle.h"
+#include "GDKTaskQueueHelpers.h"
+#include "GDKThreadCheck.h"
+#include "GDKHandleTracker.h"
+#include "GDKUserManager.h"
+#include "GenericPlatform/GenericPlatformChunkInstall.h"
+#include "Modules/ModuleManager.h"
+#include "CoreGlobals.h"
+#include "Misc/MessageDialog.h"
+#include "Misc/ConfigCacheIni.h"
+#include "Misc/CoreDelegates.h"
+#include "Misc/Optional.h"
+#include "Misc/Paths.h"
+#include "Misc/CommandLine.h"
+#include "Math/NumericLimits.h"
+#include "Stats/Stats.h"
+#include "HAL/FileManager.h"
+#include "HAL/ThreadSafeBool.h"
+#include "HAL/PlatformMisc.h"
+#include "HAL/PlatformProcess.h"
+#include "HAL/PlatformAtomics.h"
+#include "Templates/UniquePtr.h"
+#include "Containers/Ticker.h"
+#include "Microsoft/COMPointer.h"
+#include "Async/Async.h"
+
+#include "Microsoft/AllowMicrosoftPlatformTypes.h"
+THIRD_PARTY_INCLUDES_START
+#include <XGameUI.h>
+#include <XGameRuntimeInit.h>
+#include <XGameRuntimeFeature.h>
+#include <XGame.h>
+#include <XSystem.h>
+#include <XError.h>
+#include <XGameErr.h>
+#include <XUser.h>
+#include <XNetworking.h>
+#include <XThread.h>
+#include <XDisplay.h>
+#include <XGameInvite.h>
+#include <ipifcons.h>
+#include <grdk.h>
+
+#if _GRDK_EDITION >= 230300
+#include <XGameProtocol.h>
+#endif
+THIRD_PARTY_INCLUDES_END
+#include "Microsoft/HideMicrosoftPlatformTypes.h"
+
+
+// XGameProtocol support added in 230300
+#if !defined(GDK_PROTOCOL_SUPPORT_XGAMEPROTOCOL)
+	#define GDK_PROTOCOL_SUPPORT_XGAMEPROTOCOL (_GRDK_EDITION >= 230300)
+#endif
+
+// GDK 230300 documentation says ms-xbl-<titleid> is now reserved for game invites. 
+// However, newer OS versions also use it for "Play Here Instead" passing "ms-xbl-<titleid>:///" on the command line
+#if !defined(GDK_PROTOCOL_SUPPORT_MS_XBL)
+	#define GDK_PROTOCOL_SUPPORT_MS_XBL 1
+#endif
+
+static TAutoConsoleVariable<FString> CVarGDKDefaultGamertagComponent(
+	TEXT("GDK.DefaultGamertagComponent"),
+	"Classic",
+	TEXT("Choose one from [Classic, Modern, ModernSuffix, UniqueModern]. Refer to GDK document for more info.")
+);
+
+#if GDK_PROTOCOL_SUPPORT_MS_XBL
+FString GDKTitleIdUriPrefix = TEXT("ms-xbl://");
+static FAutoConsoleVariableRef CVarGDKTitleIdUriPrefix(
+	TEXT("GDK.TitleIdUriPrefix"),
+	GDKTitleIdUriPrefix,
+	TEXT("Prefix for protocol activation Uri when triggered via ms-xbl-<titleid>://")
+);
+#endif //GDK_PROTOCOL_SUPPORT_MS_XBL
+
+
+
+
+DEFINE_LOG_CATEGORY(LogGDK);
+
+
+class FGDKRuntimeModule : public IGDKRuntimeModule
+{
+public:
+	virtual void StartupModule() override
+	{
+		Init();
+	}
+	virtual void ShutdownModule() override
+	{
+		Teardown();
+	}
+
+	virtual bool IsAvailable() const override;
+
+	virtual EAppReturnType::Type MessageBoxExt(EAppMsgType::Type MsgType, const TCHAR* Text, const TCHAR* Caption) const override;
+
+	virtual ENetworkConnectionType GetNetworkConnectionType() const override;
+	virtual bool HasActiveWiFiConnection() const override;
+
+	virtual uint32 GetTitleId() const override;
+	virtual FString GetXboxSandboxId() const override;
+	virtual FString GetPrimaryServiceConfigId(bool bAllowPlaceholder) const override;
+	virtual FString GetGamertag(const FGDKUserHandle& UserHandle) const override;
+	virtual FString GetGamertag(const FGDKUserHandle& UserHandle, XUserGamertagComponent GamertagComponent) const override;
+	virtual XUserGamertagComponent GetDefaultGamertagComponent() const override;
+
+	virtual const TCHAR* GetConfigSectionName() const override;
+	virtual bool IsUsingSimplifiedUserModel() const override;
+	virtual bool IsNetworkInitialized() const override;
+
+	virtual void DisableScreenTimeout() override;
+	virtual void EnableScreenTimeout() override;
+
+	virtual void Internal_HandlePlatformTextFieldVisible( bool bIsVisible ) override;
+	virtual bool Internal_IsPlatformTextFieldVisible() const override;
+
+	virtual HRESULT AsyncGDKTask( TFunction<HRESULT(XAsyncBlock*)> InitFunction, TFunction<void(XAsyncBlock*)> ResultFunction, XTaskQueueHandle TaskQueue ) const override;
+	virtual HRESULT AsyncGDKTask( class FGDKAsyncTaskMonitor& OutMonitor, TFunction<HRESULT(XAsyncBlock*)> InitFunction, TFunction<void(XAsyncBlock*)> ResultFunction, XTaskQueueHandle TaskQueue ) const override;
+	virtual XTaskQueueHandle GetGenericTaskQueue() const override;
+	virtual XTaskQueueHandle GetBackgroundTaskQueue() const override;
+
+
+	virtual FGDKUserHandle GetUserHandleByPlatformId(FPlatformUserId PlatformUserId) const override;
+	virtual FGDKUserHandle GetUserHandleByXUserId(uint64 XUserId) const override;
+	virtual FPlatformUserId GetPlatformIdByUserHandle(const FGDKUserHandle& UserHandle) override;
+	virtual TOptional<uint64> GetXUserIdByPlatformId(FPlatformUserId PlatformUserId) const override;
+	virtual TArray<FGDKUserHandle> GetAllUserHandles() const override;
+	virtual bool WasUserRecentlySignedOut(XUserLocalId UserLocalId) const override;
+	virtual void PickUserAsync(XUserAddOptions Options, FGDKPickUserCompleteDelegate Delegate) const override;
+	virtual FPlatformUserId GetUnpairedPlatformId() const override;
+
+	virtual void SetRuntimeErrorProcessDelegate(HRESULT hResult, FGDKRuntimeErrorProcessDelegate Delegate) override;
+	virtual void ClearRuntimeErrorProcessDelegate(HRESULT hResult) override;
+
+	virtual FDelegateHandle RegisterGameInviteReceivedDelegate( FGDKOnGameInviteReceivedDelegate Delegate ) override;
+	virtual void UnregisterGameInviteReceivedDelegate( const FDelegateHandle& DelegateHandle ) override;
+
+
+	FGDKUserManager& GetUserManager() const { return *UserManagerInstance.Get(); }
+private:
+
+	static void OnNetworkConnectivityChangeCallback( void*, const struct XNetworkingConnectivityHint* ConnectivityHint );
+#if GDK_PROTOCOL_SUPPORT_XGAMEPROTOCOL || GDK_PROTOCOL_SUPPORT_MS_XBL
+	static void OnProtocolActivationCallback( void*, const char* ProtocolUri );	
+#endif
+	void SetNetworkConnectionType( const struct XNetworkingConnectivityHint* ConnectivityHint );
+	void SetDefaultGamertagComponent(IConsoleVariable* CVar);
+
+	void SetProtocolActivationUri(const FString& NewUriString);
+	bool TranslateUri( const FString& SourceUri, FString& OutDstUri ) const;
+#if GDK_PROTOCOL_SUPPORT_XGAMEPROTOCOL || GDK_PROTOCOL_SUPPORT_MS_XBL
+	bool ParseAndRemoveProtocolUri( FString& InOutCommandLine );
+#endif
+	static void OnGameInviteReceivedCallback( void*, const char* ProtocolUri );	
+
+	void Init();
+	void Teardown();
+	void CacheConfiguration();
+	bool bGDKEnvironmentInitialized = false;
+	bool bHasMSGamingRuntimeModule = false;
+
+	TUniquePtr<FGDKUserManager> UserManagerInstance;
+
+	ENetworkConnectionType CachedNetworkConnectionType = ENetworkConnectionType::Unknown;
+	TOptional<XTaskQueueRegistrationToken> NetworkConnectivityEventToken;
+	bool bCachedNetworkInitialized = false;
+	TOptional<FDelegateHandle> NetworkInitializedSuspendHandle;
+	TOptional<FDelegateHandle> NetworkInitializedResumeHandle;
+
+	FString ProtocolActivationUri;
+	FTSTicker::FDelegateHandle ProtocolActivationTickerHandle;
+#if GDK_PROTOCOL_SUPPORT_XGAMEPROTOCOL
+	TOptional<XTaskQueueRegistrationToken> ProtocolActivationRegistrationToken;
+	TMap<FString,FString> UriActivationProtocols;
+#endif
+	TOptional<XTaskQueueRegistrationToken> GameInviteRegistrationToken;
+	FGDKOnGameInviteReceived OnGameInviteReceivedDelegate;
+	FString LastGameInviteRecivedUri;
+
+
+	XDisplayTimeoutDeferralHandle TimeoutDeferralHandle = nullptr;
+	XUserGamertagComponent DefaultGamertagComponent = XUserGamertagComponent::Classic;
+
+	int32 PlatformTextFieldVisibleCounter = 0;
+
+	struct FGDKEnvironment
+	{
+		uint32 TitleId = 0;
+		FString XboxLiveSandboxId;
+		FString PrimaryServiceConfigId;
+		bool bHasValidPrimaryServiceConfigId = false;
+	} GDKEnvironment;
+
+	TMap<HRESULT, FGDKRuntimeErrorProcessDelegate> RuntimeErrorProcessDelegates;
+	FCriticalSection RuntimeErrorProcessDelegatesLock;
+
+#if WITH_EDITOR
+	virtual void Internal_InitForPIE() override;
+	virtual void Internal_TeardownForPIE() override;
+	virtual TMulticastDelegate<void()>& GetOnInitForPIE() override { return OnInitForPIEDelegate; }
+	virtual TMulticastDelegate<void()>& GetOnTeardownForPIE() override { return OnTeardownForPIEDelegate; }
+
+	TMulticastDelegate<void()> OnInitForPIEDelegate;
+	TMulticastDelegate<void()> OnTeardownForPIEDelegate;
+#endif
+
+};
+
+IMPLEMENT_MODULE(FGDKRuntimeModule, GDKRuntime);
+
+
+
+
+
+
+static FAutoConsoleCommand CmdGDKDebugDump(
+	TEXT("GDK.DumpUsers"),
+	TEXT("Logs out all current user states"),
+	FConsoleCommandWithArgsDelegate::CreateLambda( [](const TArray<FString>& Args )
+	{
+		FGDKRuntimeModule& RuntimeInstance = ((FGDKRuntimeModule&)IGDKRuntimeModule::Get());
+		RuntimeInstance.GetUserManager().DebugDump();
+	})
+);
+
+
+static FAutoConsoleCommand CmdGDKDebugAddUser(
+	TEXT("GDK.AddDebugUser"),
+	TEXT("Adds a new XUser - debug only!"),
+	FConsoleCommandWithArgsDelegate::CreateLambda( [](const TArray<FString>& Args )
+	{
+		IGDKRuntimeModule::Get().PickUserAsync(XUserAddOptions::AllowGuests, nullptr);
+	})
+);
+
+
+
+
+
+
+
+void FGDKRuntimeModule::Init()
+{
+#if PLATFORM_WINDOWS
+	// cache whether the MSGamingRuntime module is loaded. This determines where the GDK target settings are read from
+	bHasMSGamingRuntimeModule = FModuleManager::Get().ModuleExists(TEXT("MSGamingRuntime"));
+#else
+	bHasMSGamingRuntimeModule = false;
+#endif
+
+	// initialize GDK runtime
+	bGDKEnvironmentInitialized = SetupGDKEnvironment(GetConfigSectionName());
+	if (!bGDKEnvironmentInitialized)
+	{
+		return;
+	}
+
+
+	// init task queues
+	FGDKAsyncTaskQueue::PlatformInit();
+
+
+#if WITH_GDK_THREAD_CHECK
+	if (XGameRuntimeIsFeatureAvailable(XGameRuntimeFeature::XThread) && IsGDKTimeSensitiveThreadCheckEnabled())
+	{
+		// mark the game thread as time sensitive
+		check(IsInGameThread());
+		XThreadSetTimeSensitive(true);
+
+		// mark RHI and RenderThreads as time sensitive once they are up and running
+		AsyncTask( ENamedThreads::RHIThread,[]()
+		{ 
+			XThreadSetTimeSensitive(true);
+		});
+
+		AsyncTask( ENamedThreads::ActualRenderingThread,[]()
+		{ 
+			XThreadSetTimeSensitive(true);
+		});
+	}
+#endif
+
+	SetDefaultGamertagComponent(CVarGDKDefaultGamertagComponent.AsVariable());
+	CVarGDKDefaultGamertagComponent.AsVariable()->OnChangedDelegate().AddRaw(this, &FGDKRuntimeModule::SetDefaultGamertagComponent);
+
+	// Error callback message hook
+	if (XGameRuntimeIsFeatureAvailable(XGameRuntimeFeature::XError))
+	{
+		GDK_SCOPE_NOT_TIME_SENSITIVE(); // (startup only) XErrorSetOptions & XErrorSetCallback are not safe to call on time-sensitive threads
+		XErrorSetOptions(XErrorOptions::DebugBreakOnError | XErrorOptions::OutputDebugStringOnError, XErrorOptions::OutputDebugStringOnError);
+		XErrorSetCallback([](HRESULT hResult, const char* Msg, void* Context)
+		{
+			if (FGDKRuntimeModule* This = reinterpret_cast<FGDKRuntimeModule*>(Context))
+			{
+				FScopeLock ScopeLock(&This->RuntimeErrorProcessDelegatesLock);
+				FGDKRuntimeErrorProcessDelegate* DelegatePtr = This->RuntimeErrorProcessDelegates.Find(hResult);
+				if (DelegatePtr != nullptr)
+				{
+					FGDKRuntimeErrorProcessDelegate Delegate = *DelegatePtr;
+					if (Delegate.Execute(Msg))
+					{
+						return true;
+					}
+				}
+			}
+
+			if (hResult == E_OUTOFMEMORY)
+			{
+				// Cannot allocate memory during error handling if we are out of memory.  We will recurse and eventually exceed the stack.
+				char OutOfMemoryMessageBuffer[2 * 1024];
+				TCString<char>::Snprintf(OutOfMemoryMessageBuffer, sizeof(OutOfMemoryMessageBuffer), "XGameError: 0x%X - %s.\n", hResult, Msg ? Msg : "(no message)");
+				OutputDebugStringA(OutOfMemoryMessageBuffer);
+			}
+			else
+			{
+				UE_LOG( LogGDK, Error, TEXT("XGameError: 0x%X - %s."), hResult, Msg ? UTF8_TO_TCHAR(Msg) : TEXT("(no message)"));
+			}
+			return true;
+		}, this);
+	}
+
+
+
+	// Network connection change hooks
+	if (XGameRuntimeIsFeatureAvailable(XGameRuntimeFeature::XNetworking))
+	{
+		XTaskQueueRegistrationToken Token;
+		if (SUCCEEDED(XNetworkingRegisterConnectivityHintChanged( FGDKAsyncTaskQueue::GetGenericQueue(), this, &FGDKRuntimeModule::OnNetworkConnectivityChangeCallback, &Token)))
+		{
+			NetworkConnectivityEventToken = Token;
+		}
+
+		// Get current network initialization status (very likely not initialized)
+		XNetworkingConnectivityHint ConnectivityHint = {};
+		const HRESULT GetConnectivityHintResult = XNetworkingGetConnectivityHint(&ConnectivityHint);
+		if (SUCCEEDED(GetConnectivityHintResult))
+		{
+			SetNetworkConnectionType(&ConnectivityHint);
+		}
+		else
+		{
+			UE_LOG(LogGDK, Error, TEXT("Failed to determine new network initialization status due to error. Error=[%0.8X]"), GetConnectivityHintResult);
+			FPlatformMisc::SetNetworkConnectionStatus(ENetworkConnectionStatus::Disabled);
+		}
+
+		// Bind to application suspend delegate so we can update the network initialization status when we suspend
+		NetworkInitializedSuspendHandle = FCoreDelegates::ApplicationWillEnterBackgroundDelegate.AddLambda([this]()
+		{
+			SCOPED_ENTER_BACKGROUND_EVENT(STAT_Networking_ApplicationWillEnterBackgroundLambda);
+
+			// Pretend the network is gone as we're about to suspend
+			if (bCachedNetworkInitialized != false)
+			{
+				bCachedNetworkInitialized = false;
+				FCoreDelegates::ApplicationNetworkInitializationChanged.Broadcast(bCachedNetworkInitialized);
+			}
+		});
+
+		// Bind to application resume delegate so we can update the network initialization status when we resume
+		NetworkInitializedResumeHandle = FCoreDelegates::ApplicationHasEnteredForegroundDelegate.AddLambda([this]()
+		{
+			if (XGameRuntimeIsFeatureAvailable(XGameRuntimeFeature::XNetworking))
+			{
+				XNetworkingConnectivityHint ConnectivityHint = {};
+				const HRESULT GetConnectivityHintLambdaResult = XNetworkingGetConnectivityHint(&ConnectivityHint);
+				if (SUCCEEDED(GetConnectivityHintLambdaResult))
+				{
+					SetNetworkConnectionType(&ConnectivityHint);
+				}
+				else
+				{
+					UE_LOG(LogGDK, Error, TEXT("Failed to determine new network initialization status due to error. Error=[%0.8X]"), GetConnectivityHintLambdaResult);
+				}
+			}
+			else
+			{
+				SetNetworkConnectionType(nullptr);
+			}
+		});
+	}
+	else
+	{
+		SetNetworkConnectionType(nullptr);
+	}
+
+	// protocol activation
+#if GDK_PROTOCOL_SUPPORT_XGAMEPROTOCOL
+	if (XGameRuntimeIsFeatureAvailable(XGameRuntimeFeature::XGame))
+	{
+		GDK_SCOPE_NOT_TIME_SENSITIVE(); // (startup only) XGameProtocolRegisterForActivation is not safe to call on time-sensitive threads
+		XTaskQueueRegistrationToken Token;
+		if (SUCCEEDED(XGameProtocolRegisterForActivation( FGDKAsyncTaskQueue::GetGenericQueue(), this, FGDKRuntimeModule::OnProtocolActivationCallback, &Token )))
+		{
+			ProtocolActivationRegistrationToken = Token;
+		}
+	}
+#endif
+
+	// game invite hooking (not used for actual game invites, just accessing ms-xbl-xxx)
+	if (XGameRuntimeIsFeatureAvailable(XGameRuntimeFeature::XGameInvite))
+	{
+		GDK_SCOPE_NOT_TIME_SENSITIVE(); // (startup only) XGameInviteRegisterForEvent is not safe to call on time-sensitive thread
+		XTaskQueueRegistrationToken Token;
+		if (SUCCEEDED(XGameInviteRegisterForEvent( FGDKAsyncTaskQueue::GetGenericQueue(), this, FGDKRuntimeModule::OnGameInviteReceivedCallback, &Token )))
+		{
+			GameInviteRegistrationToken = Token;
+		}
+	}
+
+
+	// create the GDK user manager
+	UserManagerInstance.Reset(new FGDKUserManager);
+
+	// attempt to add the default user (i.e the user that launched the title on Xbox, or the user signed in to the xbox game bar on Windows)
+	// NB. not done in the GDK user manager constructor because there's a chance it will trigger OSS delegates that call back to FGDKUserManager
+	UserManagerInstance->AddDefaultUser();
+
+
+	// Environment Information.
+	if (GConfig != nullptr && GConfig->IsReadyForUse())
+	{
+		CacheConfiguration();
+	}
+	else
+	{
+		FCoreDelegates::TSConfigReadyForUse().AddRaw(this, &FGDKRuntimeModule::CacheConfiguration);
+	}
+}
+
+void FGDKRuntimeModule::Teardown()
+{
+	if (bGDKEnvironmentInitialized)
+	{
+#if GDK_PROTOCOL_SUPPORT_XGAMEPROTOCOL
+		if (XGameRuntimeIsFeatureAvailable(XGameRuntimeFeature::XGame) && ProtocolActivationRegistrationToken.IsSet())
+		{
+			GDK_SCOPE_NOT_TIME_SENSITIVE(); // (shutdown only) XGameProtocolUnregisterForActivation is not safe to call on time-sensitive threads
+			XGameProtocolUnregisterForActivation(ProtocolActivationRegistrationToken.GetValue(), true);
+			ProtocolActivationRegistrationToken.Reset();
+		}
+#endif
+
+		if (XGameRuntimeIsFeatureAvailable(XGameRuntimeFeature::XGameInvite) && GameInviteRegistrationToken.IsSet())
+		{
+			GDK_SCOPE_NOT_TIME_SENSITIVE(); // (shutdown only) XGameInviteUnregisterForEvent is not safe to call on time-sensitive threads
+			XGameInviteUnregisterForEvent(GameInviteRegistrationToken.GetValue(), true);
+			GameInviteRegistrationToken.Reset();
+		}
+
+		if (XGameRuntimeIsFeatureAvailable(XGameRuntimeFeature::XNetworking) && NetworkConnectivityEventToken.IsSet() )
+		{
+			GDK_SCOPE_NOT_TIME_SENSITIVE(); // (shutdown only) XNetworkingUnregisterConnectivityHintChanged is not safe to call on time-sensitive threads
+			XNetworkingUnregisterConnectivityHintChanged(NetworkConnectivityEventToken.GetValue(), true);
+			NetworkConnectivityEventToken.Reset();
+		}
+
+		if (XGameRuntimeIsFeatureAvailable(XGameRuntimeFeature::XError))
+		{
+			GDK_SCOPE_NOT_TIME_SENSITIVE(); // (shutdown only) XErrorSetCallback is not safe to call on time-sensitive threads
+			XErrorSetCallback(nullptr, nullptr);
+		}
+
+		// shut down task queues
+		FGDKAsyncTaskQueue::PlatformTeardown();
+
+		// destroy the GDK user manager
+		UserManagerInstance.Reset();
+	}
+
+	// clean up GDK
+	TeardownGDKEnvironment();
+	bGDKEnvironmentInitialized = false;
+
+	// remove ticker for uri protocol handling
+	if (ProtocolActivationTickerHandle.IsValid())
+	{
+		FTSTicker::GetCoreTicker().RemoveTicker(ProtocolActivationTickerHandle);
+		ProtocolActivationTickerHandle.Reset();
+	}
+}
+
+bool FGDKRuntimeModule::IsAvailable() const
+{
+	return bGDKEnvironmentInitialized;
+}
+
+ENetworkConnectionType FGDKRuntimeModule::GetNetworkConnectionType() const
+{
+	return CachedNetworkConnectionType;
+}
+
+bool FGDKRuntimeModule::HasActiveWiFiConnection() const
+{
+	return CachedNetworkConnectionType == ENetworkConnectionType::WiFi;
+}
+
+void FGDKRuntimeModule::OnNetworkConnectivityChangeCallback( void* Context, const XNetworkingConnectivityHint* ConnectivityHint )
+{
+	((FGDKRuntimeModule*)Context)->SetNetworkConnectionType(ConnectivityHint);
+}
+
+void FGDKRuntimeModule::SetNetworkConnectionType( const XNetworkingConnectivityHint* ConnectivityHint )
+{
+	ENetworkConnectionType OldConnectionType = CachedNetworkConnectionType;
+
+	bool bOldNetworkInitialization = bCachedNetworkInitialized;
+	bCachedNetworkInitialized = ConnectivityHint && ConnectivityHint->networkInitialized;
+
+	if (bCachedNetworkInitialized)
+	{
+		if (ConnectivityHint->ianaInterfaceType == IF_TYPE_ETHERNET_CSMACD)
+		{
+			UE_CLOG( (OldConnectionType != CachedNetworkConnectionType), LogGDK, Log, TEXT("Network connectivity level: Ethernet") );
+			CachedNetworkConnectionType = ENetworkConnectionType::Ethernet;
+		}
+		else if (ConnectivityHint->ianaInterfaceType == IF_TYPE_IEEE80211)
+		{
+			UE_CLOG( (OldConnectionType != CachedNetworkConnectionType), LogGDK, Log, TEXT("Network connectivity level: Wifi") );
+			CachedNetworkConnectionType = ENetworkConnectionType::WiFi;
+		}
+	}
+	else if (CachedNetworkConnectionType != ENetworkConnectionType::Unknown)
+	{
+		UE_CLOG( (OldConnectionType != CachedNetworkConnectionType), LogGDK, Log, TEXT("Network connectivity level: Unknown") );
+		CachedNetworkConnectionType = ENetworkConnectionType::Unknown;
+	}
+
+	if (ConnectivityHint == nullptr)
+	{
+		// If we can't tell, assume the network is available
+		bCachedNetworkInitialized = true;
+	}
+
+	const ENetworkConnectionStatus NewNetworkConnectionStatus = (ConnectivityHint && ConnectivityHint->connectivityLevel == XNetworkingConnectivityLevelHint::InternetAccess) ? ENetworkConnectionStatus::Connected : ENetworkConnectionStatus::Disabled;
+	FPlatformMisc::SetNetworkConnectionStatus(NewNetworkConnectionStatus);
+
+	if (bOldNetworkInitialization != bCachedNetworkInitialized)
+	{
+		FCoreDelegates::ApplicationNetworkInitializationChanged.Broadcast(bCachedNetworkInitialized);
+	}
+}
+
+
+void FGDKRuntimeModule::CacheConfiguration()
+{
+	GDK_SCOPE_NOT_TIME_SENSITIVE(); // (startup only) XGameGetXboxTitleId & XGameGetXboxLiveSandboxId are not safe to call on a time-critical thread
+
+	// title id
+	HRESULT hResult;
+	hResult = XGameGetXboxTitleId(&GDKEnvironment.TitleId);
+	UE_CLOG(FAILED(hResult), LogGDK, Display, TEXT("failed to get TitleId: %x"), hResult);
+
+	// xbl sandbox id
+	char SandboxIdBuf[XSystemXboxLiveSandboxIdMaxBytes] = { 0 };
+	hResult = XSystemGetXboxLiveSandboxId(XSystemXboxLiveSandboxIdMaxBytes, SandboxIdBuf, nullptr);
+	UE_CLOG(FAILED(hResult), LogGDK, Display, TEXT("failed to get SandboxId: %x"), hResult);
+	if (SUCCEEDED(hResult))
+	{
+		GDKEnvironment.XboxLiveSandboxId = UTF8_TO_TCHAR(SandboxIdBuf);
+	}
+
+	// xbl service config id
+	check(GConfig && GConfig->IsReadyForUse());
+	FString ConfiguredSCID;
+	if (GConfig->GetString(TEXT("OnlineSubsystemGDK"), TEXT("PrimaryServiceConfigId"), ConfiguredSCID, GEngineIni) && !ConfiguredSCID.IsEmpty())
+	{
+		GDKEnvironment.PrimaryServiceConfigId = ConfiguredSCID;
+		GDKEnvironment.bHasValidPrimaryServiceConfigId = true;
+	}
+	else if (GConfig->GetString(GetConfigSectionName(), TEXT("PrimaryServiceConfigId"), ConfiguredSCID, GEngineIni) && !ConfiguredSCID.IsEmpty())
+	{
+		GDKEnvironment.PrimaryServiceConfigId = ConfiguredSCID;
+		GDKEnvironment.bHasValidPrimaryServiceConfigId = true;
+	}
+	else
+	{
+		GDKEnvironment.PrimaryServiceConfigId = FString::Printf(TEXT("00000000-0000-0000-0000-0000%08X"), GetTitleId());
+		GDKEnvironment.bHasValidPrimaryServiceConfigId = false;
+	}
+
+	// show results
+	UE_LOG(LogGDK, Display, TEXT("TitleId: %x"), GDKEnvironment.TitleId);
+	UE_LOG(LogGDK, Display, TEXT("SandboxId: %s"), *GDKEnvironment.XboxLiveSandboxId );
+	UE_LOG(LogGDK, Display, TEXT("Primary Service Configuration Id: %s"), *GDKEnvironment.PrimaryServiceConfigId );
+	UE_LOG(LogGDK, Display, TEXT("GDK Edition: %d"), _GRDK_EDITION);
+
+
+	// cache activation protocols
+#if GDK_PROTOCOL_SUPPORT_XGAMEPROTOCOL
+	TArray<FString> RawActivationProtocols;
+	GConfig->GetArray(GetConfigSectionName(), TEXT("UriActivationProtocols"), RawActivationProtocols, GEngineIni);
+	for (FString& RawActivationProtocol : RawActivationProtocols)
+	{
+		RawActivationProtocol.TrimStartAndEndInline();
+		RawActivationProtocol.TrimCharInline(TEXT('('), nullptr);
+		RawActivationProtocol.TrimCharInline(TEXT(')'), nullptr);
+
+		FString RegisteredProtocol;
+		FString InGameUriPrefix;
+		if (FParse::Value(*RawActivationProtocol, TEXT("RegisteredProtocol="), RegisteredProtocol) && FParse::Value(*RawActivationProtocol, TEXT("InGameUriPrefix="), InGameUriPrefix ) )
+		{
+			RegisteredProtocol.Split(TEXT(":"), &RegisteredProtocol, nullptr); // ignore any suffix in case someone has added a :// etc.
+			RegisteredProtocol.ToLowerInline();                                // protocol must be lower case
+
+			UriActivationProtocols.Add(RegisteredProtocol, InGameUriPrefix);
+		}
+	}
+
+	// check command line now
+#if GDK_PROTOCOL_SUPPORT_MS_XBL
+	// cvars may not be initialized yet, so read the value directly
+	GConfig->GetString(TEXT("ConsoleVariables"), TEXT("GDK.TitleIdUriPrefix"), GDKTitleIdUriPrefix, GEngineIni);
+
+	FString CmdLine(FCommandLine::Get());
+	if (ParseAndRemoveProtocolUri(CmdLine))
+	{
+		FCommandLine::Set(*CmdLine);
+	}
+#endif //GDK_PROTOCOL_SUPPORT_MS_XBL
+#endif //GDK_PROTOCOL_SUPPORT_XGAMEPROTOCOL
+}
+
+
+uint32 FGDKRuntimeModule::GetTitleId() const
+{
+	return GDKEnvironment.TitleId;
+}
+
+FString FGDKRuntimeModule::GetXboxSandboxId() const
+{
+	return GDKEnvironment.XboxLiveSandboxId;
+}
+
+FString FGDKRuntimeModule::GetPrimaryServiceConfigId(bool bAllowPlaceholder) const
+{
+	if (GDKEnvironment.bHasValidPrimaryServiceConfigId || bAllowPlaceholder)
+	{
+		return GDKEnvironment.PrimaryServiceConfigId;
+	}
+	else
+	{
+		return FString();
+	}
+}
+
+void LexFromString(XUserGamertagComponent& Value, const TCHAR* ComponentName)
+{
+	if (FCString::Stricmp(ComponentName, TEXT("Classic")) == 0)
+	{
+		Value = XUserGamertagComponent::Classic;
+	}
+	else if (FCString::Stricmp(ComponentName, TEXT("Modern")) == 0)
+	{
+		Value = XUserGamertagComponent::Modern;
+	}
+	else if (FCString::Stricmp(ComponentName, TEXT("ModernSuffix")) == 0)
+	{
+		Value = XUserGamertagComponent::ModernSuffix;
+	}
+	else if (FCString::Stricmp(ComponentName, TEXT("UniqueModern")) == 0)
+	{
+		Value = XUserGamertagComponent::UniqueModern;
+	}
+	else
+	{
+		UE_LOG(LogGDK, Warning, TEXT("[LexFromString] ComponentName FString [%s] did not match any XUserGamertagComponent. Returning XUserGamertagComponent::Classic as default."), ComponentName);
+
+		Value = XUserGamertagComponent::Classic;
+	}
+}
+
+size_t GetMaxBytesForGamertagComponent(XUserGamertagComponent GamertagComponent)
+{
+	switch (GamertagComponent)
+	{
+	case XUserGamertagComponent::Classic:
+		return XUserGamertagComponentClassicMaxBytes;
+	case XUserGamertagComponent::Modern:
+		return XUserGamertagComponentModernMaxBytes;
+	case XUserGamertagComponent::ModernSuffix:
+		return XUserGamertagComponentModernSuffixMaxBytes;
+	case XUserGamertagComponent::UniqueModern:
+		return XUserGamertagComponentUniqueModernMaxBytes;
+	default:
+		checkNoEntry();
+		return 0Ui64;
+	}
+}
+
+void FGDKRuntimeModule::SetDefaultGamertagComponent(IConsoleVariable* CVar)
+{
+	FString UserGamertagComponentStr = CVar->AsVariable()->GetString();
+	LexFromString(DefaultGamertagComponent, *UserGamertagComponentStr);
+}
+
+FString FGDKRuntimeModule::GetGamertag(const FGDKUserHandle& UserHandle) const
+{
+	return GetGamertag(UserHandle, DefaultGamertagComponent);
+}
+
+FString FGDKRuntimeModule::GetGamertag(const FGDKUserHandle& UserHandle, XUserGamertagComponent GamertagComponent) const
+{
+	if (!UserHandle.IsValid())
+	{
+		return TEXT("(invalid user)");
+	}
+
+	FString Gamertag;
+
+	size_t GamertagComponentMaxBytes = GetMaxBytesForGamertagComponent(GamertagComponent);
+	TUniquePtr<char[]> GamerTagBuffer = MakeUnique<char[]>(GamertagComponentMaxBytes + 1);
+	size_t GamerTagBufferUsedSize = 0;
+	HRESULT Result = XUserGetGamertag(UserHandle, GamertagComponent, GamertagComponentMaxBytes, GamerTagBuffer.Get(), &GamerTagBufferUsedSize);
+	if (SUCCEEDED(Result))
+	{
+		Gamertag = UTF8_TO_TCHAR(GamerTagBuffer.Get());
+	}
+	else
+	{
+		UE_LOG(LogGDK, Warning, TEXT("[FGDKRuntimeModule::GetGamertag] Failed to get user gamertag. Result = (0x%0.8X)"), Result);
+	}
+
+	return Gamertag;
+}
+
+XUserGamertagComponent FGDKRuntimeModule::GetDefaultGamertagComponent() const 
+{
+	return DefaultGamertagComponent;
+}
+
+const TCHAR* FGDKRuntimeModule::GetConfigSectionName() const
+{
+#if PLATFORM_WINDOWS
+	static FString ConfigSectionName = [this]()
+	{
+		return bHasMSGamingRuntimeModule ? TEXT("/Script/MSGamingSupport.MSGamingSettings") : FPlatformProperties::GetRuntimeSettingsClassName();
+	}();
+
+	return *ConfigSectionName;
+#else
+	return FPlatformProperties::GetRuntimeSettingsClassName();
+#endif
+}
+
+bool FGDKRuntimeModule::IsUsingSimplifiedUserModel() const
+{
+	static bool bIsUsingSimplifiedUserModel = []()
+	{
+		bool bResult = false;
+		GConfig->GetBool(IGDKRuntimeModule::Get().GetConfigSectionName(), TEXT("bUseSimplifiedUserModel"), bResult, GEngineIni);
+		return bResult;
+	}();
+
+	return bIsUsingSimplifiedUserModel;
+}
+
+bool FGDKRuntimeModule::IsNetworkInitialized() const
+{
+	return bCachedNetworkInitialized;
+}
+
+void FGDKRuntimeModule::DisableScreenTimeout()
+{
+	GDK_SCOPE_NOT_TIME_SENSITIVE(); // XDisplayAcquireTimeoutDeferral is not safe to call on a time-sensitive thread
+
+	if (TimeoutDeferralHandle != nullptr)
+	{
+		XDisplayCloseTimeoutDeferralHandle(TimeoutDeferralHandle);
+		TimeoutDeferralHandle = nullptr;
+	}
+
+	HRESULT hResult = XDisplayAcquireTimeoutDeferral(&TimeoutDeferralHandle);
+	UE_CLOG(FAILED(hResult), LogGDK, Display, TEXT("failed to acquire timeout deferral handle: %x"), hResult);
+}
+
+void FGDKRuntimeModule::EnableScreenTimeout()
+{
+	GDK_SCOPE_NOT_TIME_SENSITIVE(); // XDisplayCloseTimeoutDeferralHandle is not safe to call on a time-sensitive thread
+
+	XDisplayCloseTimeoutDeferralHandle(TimeoutDeferralHandle);
+}
+
+
+void FGDKRuntimeModule::Internal_HandlePlatformTextFieldVisible( bool bIsVisible )
+{
+	if (bIsVisible)
+	{
+		FPlatformAtomics::InterlockedIncrement(&PlatformTextFieldVisibleCounter);
+	}
+	else
+	{
+		FPlatformAtomics::InterlockedDecrement(&PlatformTextFieldVisibleCounter);
+		check(FPlatformAtomics::AtomicRead(&PlatformTextFieldVisibleCounter) >= 0);
+	}
+}
+
+bool FGDKRuntimeModule::Internal_IsPlatformTextFieldVisible() const
+{
+	return (FPlatformAtomics::AtomicRead(&PlatformTextFieldVisibleCounter) > 0);
+}
+
+
+
+class FGDKMessagBox
+{
+public:
+	FGDKMessagBox( const TCHAR* InText, const TCHAR* InCaption )
+		: Text(InText)
+		, Caption(InCaption)
+		, DefaultButton(XGameUiMessageDialogButton::First)
+		, CancelButton(XGameUiMessageDialogButton::Second)
+		, bIsFinished(false)
+		, Result(EAppReturnType::Cancel)
+		, NumButtons(0)
+	{
+		for( int i = 0; i < MaxButtons; i++ )
+		{
+			Buttons[i].Text = nullptr;
+		}
+
+		check( int(XGameUiMessageDialogButton::First) == 0 );
+		check( int(XGameUiMessageDialogButton::Second) == 1 );
+		check( int(XGameUiMessageDialogButton::Third) == 2 );
+	}
+
+	void AddButton( const char* InText, EAppReturnType::Type Id, bool bIsDefault = false )
+	{
+		if( NumButtons < MaxButtons )
+		{
+			FButton& Button = Buttons[NumButtons];
+			Button.Text = InText;
+			Button.Id = Id;
+
+			if( Id == EAppReturnType::Cancel )
+			{
+				CancelButton = (XGameUiMessageDialogButton)NumButtons;
+			}
+			if( bIsDefault )
+			{
+				DefaultButton = (XGameUiMessageDialogButton)NumButtons;
+			}
+
+			NumButtons++;
+		}
+	}
+
+	EAppReturnType::Type Show()
+	{ 
+		FGDKLocalTaskBlock Block;
+		if( SUCCEEDED(XGameUiShowMessageDialogAsync( Block, TCHAR_TO_UTF8(Caption), TCHAR_TO_UTF8(Text), Buttons[0].Text, Buttons[1].Text, Buttons[2].Text, DefaultButton, CancelButton )) )
+		{
+			Block.BlockUntilComplete();
+
+			XGameUiMessageDialogButton Button;
+			if (SUCCEEDED(XGameUiShowMessageDialogResult( Block, &Button)))
+			{
+				int ButtonIndex = (int)Button;
+				Result = Buttons[ButtonIndex].Id;
+			}
+
+			bIsFinished = true; 
+		}
+
+		return Result;
+	}
+
+protected:
+
+
+	struct FButton
+	{
+		const char* Text;
+		EAppReturnType::Type Id;
+	};
+
+	static const int MaxButtons = 3;
+
+	const TCHAR* Text;
+	const TCHAR* Caption;
+
+	XGameUiMessageDialogButton DefaultButton;
+	XGameUiMessageDialogButton CancelButton;
+
+	FThreadSafeBool bIsFinished;
+	EAppReturnType::Type Result;
+
+	int NumButtons;
+	FButton Buttons[MaxButtons];
+};
+
+
+EAppReturnType::Type FGDKRuntimeModule::MessageBoxExt(EAppMsgType::Type MsgType, const TCHAR* Text, const TCHAR* Caption) const
+{
+	FGDKMessagBox GDKMessageBox( Text, Caption );
+
+	//add the appropriate buttons
+	switch( MsgType )
+	{
+		case EAppMsgType::Ok:
+			GDKMessageBox.AddButton( "OK",         EAppReturnType::Ok );
+			break;
+
+		case EAppMsgType::YesNo:
+			GDKMessageBox.AddButton( "Yes",        EAppReturnType::Yes );
+			GDKMessageBox.AddButton( "No",         EAppReturnType::No );
+			break;
+
+		case EAppMsgType::OkCancel:
+			GDKMessageBox.AddButton( "OK",         EAppReturnType::Ok );
+			GDKMessageBox.AddButton( "Cancel",     EAppReturnType::Cancel );
+			break;
+
+		case EAppMsgType::YesNoCancel:
+			GDKMessageBox.AddButton( "Yes",        EAppReturnType::Yes );
+			GDKMessageBox.AddButton( "No",         EAppReturnType::No );
+			GDKMessageBox.AddButton( "Cancel",     EAppReturnType::Cancel );
+			break;
+
+		case EAppMsgType::CancelRetryContinue:
+			GDKMessageBox.AddButton( "Cancel",     EAppReturnType::Cancel );
+			GDKMessageBox.AddButton( "Retry",      EAppReturnType::Retry, true );
+			GDKMessageBox.AddButton( "Continue",   EAppReturnType::Continue );
+			break;
+
+		case EAppMsgType::YesNoYesAll:
+			GDKMessageBox.AddButton( "Yes",        EAppReturnType::Yes );
+			GDKMessageBox.AddButton( "No",         EAppReturnType::No );
+			GDKMessageBox.AddButton( "Yes To All", EAppReturnType::YesAll );
+			break;
+
+		default:
+			UE_LOG(LogGDK, Warning, TEXT("MsgBox: type not supported\n") );
+			return FGenericPlatformMisc::MessageBoxExt(MsgType, Text, Caption);
+	}
+
+	EAppReturnType::Type Result = GDKMessageBox.Show();
+	return Result;
+}
+
+
+
+
+void FGDKRuntimeModule::SetProtocolActivationUri(const FString& NewUriString)
+{
+	ProtocolActivationUri = NewUriString;
+
+	// start or stop the ticker as appropriate
+	if (!ProtocolActivationUri.IsEmpty())
+	{
+		UE_LOG(LogGDK, Log, TEXT("Got protocol activation Uri : %s"), *NewUriString );
+
+		if (!ProtocolActivationTickerHandle.IsValid())
+		{
+			ProtocolActivationTickerHandle = FTSTicker::GetCoreTicker().AddTicker(FTickerDelegate::CreateLambda(
+				[this](float)
+				{
+					if (FCoreDelegates::OnActivatedByProtocol.IsBound())
+					{
+						FCoreDelegates::OnActivatedByProtocol.Broadcast(ProtocolActivationUri, PLATFORMUSERID_NONE);
+
+						ProtocolActivationUri.Reset();
+						ProtocolActivationTickerHandle.Reset();
+						return false; // remove from ticker
+					}
+				
+					return true; // try again next tick
+				}
+			));
+		}
+	}
+	else if (ProtocolActivationTickerHandle.IsValid())
+	{
+		FTSTicker::GetCoreTicker().RemoveTicker(ProtocolActivationTickerHandle);
+		ProtocolActivationTickerHandle.Reset();
+	}
+}
+
+bool FGDKRuntimeModule::TranslateUri( const FString& SourceUri, FString& OutDstUri ) const
+{
+	// note: only basic Uri forms are supported at the moment    scheme:payload and scheme://payload ... no authority support
+	FString UriScheme;
+	FString UriPayload;
+	if (SourceUri.Split(FString(TEXT("://")), &UriScheme, &UriPayload) || SourceUri.Split(FString(TEXT(":")), &UriScheme, &UriPayload))
+	{
+		UriScheme.ToLowerInline();
+
+#if GDK_PROTOCOL_SUPPORT_MS_XBL
+		if (UriScheme.StartsWith(TEXT("ms-xbl-")))
+		{
+			OutDstUri = GDKTitleIdUriPrefix + UriPayload;
+			return true;
+		}
+#endif
+
+#if GDK_PROTOCOL_SUPPORT_XGAMEPROTOCOL
+		const FString* UriPrefixPtr = UriActivationProtocols.Find(UriScheme);
+		if (UriPrefixPtr != nullptr)
+		{
+			OutDstUri = (*UriPrefixPtr) + UriPayload;
+			return true;
+		}
+#endif
+	}
+
+	return false;
+}
+
+#if GDK_PROTOCOL_SUPPORT_XGAMEPROTOCOL || GDK_PROTOCOL_SUPPORT_MS_XBL
+bool FGDKRuntimeModule::ParseAndRemoveProtocolUri( FString& InOutCommandLine )
+{
+	bool bResult = false;
+
+	const TCHAR* RawCmdLine = *InOutCommandLine;
+	FString SourceUri;
+	if (FParse::Token(RawCmdLine, SourceUri, false))
+	{
+		FString TranslatedUri;
+		if (TranslateUri(SourceUri, TranslatedUri))
+		{
+			// update the command line now the Uri token has been parsed out
+			FString UpdatedCommandLine(RawCmdLine);
+			InOutCommandLine = UpdatedCommandLine;
+
+			FPlatformMisc::LowLevelOutputDebugStringf(TEXT("Launched with URI: %s\n"), *SourceUri );
+			SetProtocolActivationUri(TranslatedUri);
+
+			bResult = true;
+		}
+	}
+
+	return bResult;
+}
+#endif
+
+#if GDK_PROTOCOL_SUPPORT_XGAMEPROTOCOL || GDK_PROTOCOL_SUPPORT_MS_XBL
+void FGDKRuntimeModule::OnProtocolActivationCallback( void* Context, const char* ProtocolUri )
+{
+	FGDKRuntimeModule* RuntimeModule = (FGDKRuntimeModule*)Context;
+
+	FString SourceUri(UTF8_TO_TCHAR(ProtocolUri));
+	FString TranslatedUri;
+	if (RuntimeModule->TranslateUri(SourceUri, TranslatedUri))
+	{
+		RuntimeModule->SetProtocolActivationUri(TranslatedUri);
+	}
+}
+#endif
+
+void FGDKRuntimeModule::OnGameInviteReceivedCallback( void* Context, const char* InviteUri )
+{
+	FGDKRuntimeModule* RuntimeModule = (FGDKRuntimeModule*)Context;
+
+	// record the newest invite Uri and broadcast it to all listeners
+	RuntimeModule->LastGameInviteRecivedUri = UTF8_TO_TCHAR(InviteUri);
+	ExecuteOnGameThread( UE_SOURCE_LOCATION, [RuntimeModule, Uri = RuntimeModule->LastGameInviteRecivedUri]()
+	{
+		RuntimeModule->OnGameInviteReceivedDelegate.Broadcast(Uri);
+	});
+
+	// forward this to the normal protocol activation handler too, for deep linking & 'play here instead' support
+#if GDK_PROTOCOL_SUPPORT_XGAMEPROTOCOL || GDK_PROTOCOL_SUPPORT_MS_XBL
+	OnProtocolActivationCallback(Context, InviteUri);
+#endif
+}
+
+
+
+HRESULT FGDKRuntimeModule::AsyncGDKTask( TFunction<HRESULT(XAsyncBlock*)> InitFunction, TFunction<void(XAsyncBlock*)> ResultFunction, XTaskQueueHandle TaskQueue ) const
+{
+	return ::AsyncGDKTask(InitFunction, ResultFunction, TaskQueue);
+}
+
+HRESULT FGDKRuntimeModule::AsyncGDKTask( class FGDKAsyncTaskMonitor& OutMonitor, TFunction<HRESULT(XAsyncBlock*)> InitFunction, TFunction<void(XAsyncBlock*)> ResultFunction, XTaskQueueHandle TaskQueue ) const
+{
+	return ::AsyncGDKTask(OutMonitor, InitFunction, ResultFunction, TaskQueue);
+}
+
+XTaskQueueHandle FGDKRuntimeModule::GetGenericTaskQueue() const
+{
+	return FGDKAsyncTaskQueue::GetGenericQueue();
+}
+
+XTaskQueueHandle FGDKRuntimeModule::GetBackgroundTaskQueue() const
+{
+	return FGDKAsyncTaskQueue::GetBackgroundTaskQueue();
+}
+
+
+
+
+FGDKUserHandle FGDKRuntimeModule::GetUserHandleByPlatformId(FPlatformUserId PlatformUserId) const
+{
+	return GetUserManager().GetUserHandleByPlatformId(PlatformUserId);
+}
+
+FGDKUserHandle FGDKRuntimeModule::GetUserHandleByXUserId(uint64 XUserId) const
+{
+	return GetUserManager().GetUserHandleByXUserId(XUserId);
+}
+
+FPlatformUserId FGDKRuntimeModule::GetPlatformIdByUserHandle(const FGDKUserHandle& UserHandle)
+{
+	return GetUserManager().GetPlatformIdByUserHandle(UserHandle);
+}
+
+TOptional<uint64> FGDKRuntimeModule::GetXUserIdByPlatformId(FPlatformUserId PlatformUserId) const
+{
+	return GetUserManager().GetXUserIdByPlatformId(PlatformUserId);
+}
+
+TArray<FGDKUserHandle> FGDKRuntimeModule::GetAllUserHandles() const
+{
+	return GetUserManager().GetAllUserHandles();
+}
+
+bool FGDKRuntimeModule::WasUserRecentlySignedOut(XUserLocalId UserLocalId) const
+{
+	return GetUserManager().WasUserRecentlySignedOut(UserLocalId);
+}
+
+void FGDKRuntimeModule::PickUserAsync(XUserAddOptions Options, FGDKPickUserCompleteDelegate Delegate) const
+{
+	return GetUserManager().PickUserAsync(Options, Delegate);
+}
+
+FPlatformUserId FGDKRuntimeModule::GetUnpairedPlatformId() const
+{
+	return FPlatformUserId::CreateFromInternalId( GetUserManager().FindEmptySeat() );
+}
+
+void FGDKRuntimeModule::SetRuntimeErrorProcessDelegate(HRESULT hResult, FGDKRuntimeErrorProcessDelegate Delegate)
+{
+	FScopeLock ScopeLock(&RuntimeErrorProcessDelegatesLock);
+	RuntimeErrorProcessDelegates.Emplace(hResult, Delegate);
+}
+
+void FGDKRuntimeModule::ClearRuntimeErrorProcessDelegate(HRESULT hResult)
+{
+	FScopeLock ScopeLock(&RuntimeErrorProcessDelegatesLock);
+	RuntimeErrorProcessDelegates.Remove(hResult);
+}
+
+FDelegateHandle FGDKRuntimeModule::RegisterGameInviteReceivedDelegate( FGDKOnGameInviteReceivedDelegate Delegate )
+{
+	// if we've received an existing Uri already, send it to this new callback now
+	if (!LastGameInviteRecivedUri.IsEmpty())
+	{
+		ExecuteOnGameThread( UE_SOURCE_LOCATION, [Delegate, Uri = LastGameInviteRecivedUri]()
+		{
+			Delegate.ExecuteIfBound(Uri);
+		});
+	}
+
+	return OnGameInviteReceivedDelegate.Add(Delegate);
+}
+
+void FGDKRuntimeModule::UnregisterGameInviteReceivedDelegate( const FDelegateHandle& DelegateHandle )
+{
+	OnGameInviteReceivedDelegate.Remove(DelegateHandle);
+}
+
+
+#if WITH_EDITOR
+void FGDKRuntimeModule::Internal_InitForPIE()
+{
+	Init();
+	OnInitForPIEDelegate.Broadcast();
+}
+void FGDKRuntimeModule::Internal_TeardownForPIE()
+{
+	OnTeardownForPIEDelegate.Broadcast();
+	Teardown();
+}
+#endif
+
+
+
+FString LexToString( const FGDKUserHandle& UserHandle )
+{
+	if (!UserHandle.IsValid())
+	{
+		return TEXT("<null user>");
+	}
+
+	FString Result = IGDKRuntimeModule::Get().GetGamertag(UserHandle);
+
+	if (Result.IsEmpty())
+	{
+		Result = FString(TEXT("<invalid user>"));
+	}
+
+	return Result;
+}
+
+
+
+FString LexToString( const XUserLocalId& UserId )
+{
+	FString Result;
+	FGDKUserHandle UserHandle;
+	if (SUCCEEDED(XUserFindUserByLocalId(UserId, UserHandle.GetInitReference())))
+	{
+		Result = LexToString(UserHandle);
+	}
+	else
+	{
+		Result = FString::Printf( TEXT("<invalid id: %llx>"), UserId.value );
+	}
+
+	return Result;
+}
+
+
+const TCHAR* LexToString(XNetworkingConnectivityLevelHint Value)
+{
+	switch (Value)
+	{
+	case XNetworkingConnectivityLevelHint::None:						return TEXT("None");
+	case XNetworkingConnectivityLevelHint::ConstrainedInternetAccess:	return TEXT("ConstrainedInternetAccess");
+	case XNetworkingConnectivityLevelHint::InternetAccess:				return TEXT("InternetAccess");
+	case XNetworkingConnectivityLevelHint::LocalAccess:					return TEXT("LocalAccess");
+	case XNetworkingConnectivityLevelHint::Unknown:						return TEXT("Unknown");
+	}
+
+	checkNoEntry();
+	return TEXT("Invalid");
+}
+
+
+#endif //WITH_GRDK
diff --git a/Engine/Source/Runtime/Microsoft/GDK/GDKRuntime/Private/GDKTaskQueueHelpers.cpp b/Engine/Source/Runtime/Microsoft/GDK/GDKRuntime/Private/GDKTaskQueueHelpers.cpp
new file mode 100644
--- /dev/null
+++ b/Engine/Source/Runtime/Microsoft/GDK/GDKRuntime/Private/GDKTaskQueueHelpers.cpp
@@ -0,0 +1,490 @@
+// Copyright Epic Games, Inc. All Rights Reserved.
+
+
+#include "GDKTaskQueueHelpers.h"
+#if WITH_GRDK
+#include "GDKThreadCheck.h"
+#include "HAL/Event.h"
+#include "HAL/PlatformProcess.h"
+#include "HAL/PlatformMisc.h"
+#include "Containers/Ticker.h"
+#include "Templates/UniquePtr.h"
+
+FGDKAsyncBlock::FGDKAsyncBlock(void* InUserData, FGDKAsyncBlockDelegate InDelegate, XTaskQueueHandle TaskQueue)
+	: UserData(InUserData)
+	, Delegate(MoveTemp(InDelegate))
+{
+	// init to zero
+	AsyncBlock = XAsyncBlock{ 0 };
+
+	// default to our generic queue
+	AsyncBlock.queue = TaskQueue ? TaskQueue : FGDKAsyncTaskQueue::GetGenericQueue();
+	AsyncBlock.context = this;
+
+	// wrap the delegate with AsyncBlock callback
+	if (Delegate)
+	{
+		AsyncBlock.callback = [](XAsyncBlock* InnerBlock)
+		{
+			// get our wrapper object
+			((FGDKAsyncBlock*)InnerBlock->context)->DelegateWrapper();
+		};
+	}
+}
+
+
+
+FGDKAsyncBlock::~FGDKAsyncBlock()
+{
+	if (XAsyncGetStatus(&AsyncBlock, false) == E_PENDING)
+	{
+		UE_LOG(LogGDK, Warning, TEXT("[FGDKAsyncBlock::~FGDKAsyncBlock] AsyncBlock still pending for task. Cancelling and firing callback."));
+		XAsyncCancel(&AsyncBlock);
+	}
+}
+
+void FGDKAsyncBlock::DelegateWrapper()
+{
+	// call the delegate
+	if (Delegate)
+	{
+		Delegate(this);
+	}
+}
+
+XAsyncBlock* FGDKAsyncBlock::GetInnerBlockForGDKAPI()
+{
+	return &AsyncBlock;
+}
+
+void* FGDKAsyncBlock::GetUserData()
+{
+	return UserData;
+}
+
+HRESULT FGDKAsyncBlock::GetStatus()
+{
+	return XAsyncGetStatus(&AsyncBlock, false);
+}
+
+
+
+
+
+class FThreadPoolTaskQueue : public FGDKAsyncTaskQueue
+{
+public:
+	FThreadPoolTaskQueue() 
+		: FGDKAsyncTaskQueue(XTaskQueueDispatchMode::ThreadPool, XTaskQueueDispatchMode::Manual)
+	{
+		// set up to dispatch completions once a frame on game thread
+		TickHandle = FTSTicker::GetCoreTicker().AddTicker(FTickerDelegate::CreateLambda([this](float)
+		{
+			XTaskQueueDispatch(Queue, XTaskQueuePort::Completion, 0);
+			return true;
+		}));
+	}
+
+	~FThreadPoolTaskQueue()
+	{
+		FTSTicker::GetCoreTicker().RemoveTicker(TickHandle);
+	}
+
+private:
+	FTSTicker::FDelegateHandle TickHandle;
+};
+
+
+
+
+class FBackgroundTaskQueue : public FGDKAsyncTaskQueue
+{
+public:
+	FBackgroundTaskQueue() 
+		: FGDKAsyncTaskQueue(XTaskQueueDispatchMode::ThreadPool, XTaskQueueDispatchMode::ThreadPool)
+	{
+	}
+};
+
+
+
+
+
+static bool bIsTornDown = false;
+static TUniquePtr<FThreadPoolTaskQueue> GThreadPoolTaskQueue;
+static TUniquePtr<FBackgroundTaskQueue> GBackgroundTaskQueue;
+
+void FGDKAsyncTaskQueue::PlatformInit()
+{
+	bIsTornDown = false;
+}
+
+void FGDKAsyncTaskQueue::PlatformTeardown()
+{
+	GDK_SCOPE_NOT_TIME_SENSITIVE(); // XTaskQueueTerminate, XTaskQueueGetCurrentProcessTaskQueue and XTaskQueueSetCurrentProcessTaskQueue are not safe to call on time-sensitive threads
+
+	check(!bIsTornDown);
+	bIsTornDown = true;
+
+	// destroy our task queues
+	GThreadPoolTaskQueue.Reset();
+	GBackgroundTaskQueue.Reset();
+
+	// close the system task queue and prevent anyone from using it
+	XTaskQueueHandle SystemTaskQueue;
+	if (XTaskQueueGetCurrentProcessTaskQueue(&SystemTaskQueue) && SystemTaskQueue != nullptr)
+	{
+		XTaskQueueTerminate(SystemTaskQueue, true, nullptr, nullptr);
+		XTaskQueueCloseHandle(SystemTaskQueue);
+	}
+	CA_SUPPRESS(6387) // first parameter is marked _In_ but documentation says nullptr is allowed to disable the proces default task queue
+	XTaskQueueSetCurrentProcessTaskQueue(nullptr);
+}
+
+XTaskQueueHandle FGDKAsyncTaskQueue::GetGenericQueue()
+{
+	check(!bIsTornDown);
+
+	if (!GThreadPoolTaskQueue.IsValid())
+	{
+		GThreadPoolTaskQueue = MakeUnique<FThreadPoolTaskQueue>();
+	}
+
+	return GThreadPoolTaskQueue->GetQueue();
+}
+
+XTaskQueueHandle FGDKAsyncTaskQueue::GetBackgroundTaskQueue()
+{
+	check(!bIsTornDown);
+
+	if (!GBackgroundTaskQueue.IsValid())
+	{
+		GBackgroundTaskQueue = MakeUnique<FBackgroundTaskQueue>();
+	}
+
+	return GBackgroundTaskQueue->GetQueue();
+}
+
+
+
+
+
+
+FGDKAsyncTaskQueue::FGDKAsyncTaskQueue( XTaskQueueDispatchMode InWorkDispatchMode, XTaskQueueDispatchMode InCompletionDispatchMode )
+	: WorkDispatchMode(InWorkDispatchMode)
+	, CompletionDispatchMode(InCompletionDispatchMode)
+{
+	GDK_SCOPE_NOT_TIME_SENSITIVE(); // XTaskQueueCreate is not safe to call on time-sensitive threads
+
+	HRESULT hResult = XTaskQueueCreate(WorkDispatchMode, CompletionDispatchMode, &Queue);
+	checkf(SUCCEEDED(hResult), TEXT("XTaskQueueCreate failed: 0x%x"), hResult);
+}
+
+FGDKAsyncTaskQueue::~FGDKAsyncTaskQueue()
+{
+	CancelPendingTasksAndDestroyQueue();
+}
+
+void FGDKAsyncTaskQueue::CancelPendingTasksAndDestroyQueue()
+{
+	if (Queue != nullptr)
+	{
+		GDK_SCOPE_NOT_TIME_SENSITIVE(); // XTaskQueueTerminate is not safe to call on time-sensitive threads
+
+		// add a terminiation marker into the queue that sets an event when it has completed
+		FEvent* QueueTerminated = FPlatformProcess::GetSynchEventFromPool(false);
+		XTaskQueueTerminate(Queue, false, (void*)QueueTerminated, []( void* Context)
+		{
+			((FEvent*)Context)->Trigger();
+		});
+
+		// keep dispatching the queue until it has terminated
+		do
+		{
+			// dispatch any outstanding work
+			if (WorkDispatchMode == XTaskQueueDispatchMode::Manual)
+			{
+				XTaskQueueDispatch(Queue, XTaskQueuePort::Work, INFINITE);
+			}
+
+			// dispatch any outstanding completion - including the termination callback
+			if (CompletionDispatchMode == XTaskQueueDispatchMode::Manual)
+			{
+				XTaskQueueDispatch(Queue, XTaskQueuePort::Completion, INFINITE);
+			}
+
+		} while (!QueueTerminated->Wait(0));
+
+		// clean up
+		FPlatformProcess::ReturnSynchEventToPool(QueueTerminated);
+		XTaskQueueCloseHandle(Queue);
+		Queue = nullptr;
+	}
+}
+
+HRESULT FGDKAsyncTaskQueue::BlockUntilComplete( XAsyncBlock& AsyncBlock ) const
+{
+	check(Queue);
+
+	HRESULT hResult;
+
+	// wait until the block is complete, with special handling for game thread
+	while ((hResult = XAsyncGetStatus(&AsyncBlock, false)) == E_PENDING)
+	{
+		if (IsInGameThread())
+		{
+			FPlatformMisc::PumpMessagesOutsideMainLoop();
+		}
+
+		// dispatch any outstanding work (NB. attempting a dispatch on a non-Manual queue will result in an XError callback for illegal access)
+		if (WorkDispatchMode == XTaskQueueDispatchMode::Manual)
+		{
+			XTaskQueueDispatch(Queue, XTaskQueuePort::Work, 0);
+		}
+
+		// dispatch any outstanding completion
+		if (CompletionDispatchMode == XTaskQueueDispatchMode::Manual)
+		{
+			XTaskQueueDispatch(Queue, XTaskQueuePort::Completion, 0);
+		}
+	}
+
+	return hResult;
+}
+
+
+
+
+
+
+
+
+
+
+FGDKLocalTaskBlock::FGDKLocalTaskBlock()
+	: AsyncBlock(XAsyncBlock{ 0 })
+	, Queue(XTaskQueueDispatchMode::ThreadPool, XTaskQueueDispatchMode::Manual)
+{
+	AsyncBlock.queue = Queue.GetQueue();
+}
+
+FGDKLocalTaskBlock::~FGDKLocalTaskBlock()
+{
+	check(XAsyncGetStatus(&AsyncBlock, false) != E_PENDING);
+}
+
+HRESULT FGDKLocalTaskBlock::BlockUntilComplete()
+{
+	return Queue.BlockUntilComplete(AsyncBlock);
+}
+
+HRESULT FGDKLocalTaskBlock::GetStatus()
+{
+	return XAsyncGetStatus(&AsyncBlock, false);
+}
+
+
+
+
+
+bool FGDKAsyncTaskMonitor::TryCancel(bool bWait)
+{
+	TSharedPtr<FGDKAsyncBlock> PinnedBlock = WeakBlock.Pin();
+	if (PinnedBlock.IsValid())
+	{
+		XAsyncBlock* AsyncBlock = PinnedBlock->GetInnerBlockForGDKAPI();
+		HRESULT hResult = XAsyncGetStatus(AsyncBlock, false);
+		if (hResult == E_PENDING)
+		{
+			XAsyncCancel(AsyncBlock);
+
+			if (bWait)
+			{
+				while (hResult == E_PENDING)
+				{
+					if (IsInGameThread())
+					{
+						FPlatformMisc::PumpMessagesOutsideMainLoop();
+					}
+
+					hResult = XAsyncGetStatus(AsyncBlock, false);
+				}
+			}
+
+			return true;
+		}
+	}
+
+	return false;
+}
+
+bool FGDKAsyncTaskMonitor::IsValid() const
+{
+	return WeakBlock.IsValid() && WeakBlock.Pin().IsValid();
+}
+
+
+void Internal_InitAsyncTaskMonitor( FGDKAsyncTaskMonitor& TaskMonitor, TSharedPtr<FGDKAsyncBlock> Block )
+{
+	TaskMonitor.WeakBlock = Block;
+}
+
+
+
+
+
+HRESULT AsyncGDKTask( TFunction<HRESULT(XAsyncBlock*)> InitFunction, TFunction<void(XAsyncBlock*)> ResultFunction, XTaskQueueHandle TaskQueue )
+{
+	FGDKAsyncTaskMonitor UnusedMonitor;
+	return AsyncGDKTask( UnusedMonitor, MoveTemp(InitFunction), MoveTemp(ResultFunction), TaskQueue );
+}
+
+HRESULT AsyncGDKTask( FGDKAsyncTaskMonitor& OutAsyncTaskMonitor, TFunction<HRESULT(XAsyncBlock*)> InitFunction, TFunction<void(XAsyncBlock*)> ResultFunction, XTaskQueueHandle TaskQueue )
+{
+	// sanity check parameters
+	if (InitFunction == nullptr)
+	{
+		Internal_InitAsyncTaskMonitor( OutAsyncTaskMonitor, nullptr );
+		return ERROR_INVALID_PARAMETER;
+	}
+
+	// callback context
+	struct FContext
+	{
+		TSharedPtr<FGDKAsyncBlock> TaskBlock;
+	};
+	FContext* Context = new FContext();
+
+	// completion callback
+	auto OnGDKTaskComplete = [ResultFunction]( FGDKAsyncBlock* TaskBlock )
+	{
+		FContext* Context = (FContext*)TaskBlock->GetUserData();
+		check(TaskBlock == Context->TaskBlock.Get());
+
+		if (ResultFunction)
+		{
+			ResultFunction( *TaskBlock );
+		}
+
+		delete Context;
+	};
+
+	// create the async task block and try to get started
+	Context->TaskBlock = MakeShared<FGDKAsyncBlock>(Context, OnGDKTaskComplete, TaskQueue);
+	Internal_InitAsyncTaskMonitor( OutAsyncTaskMonitor, Context->TaskBlock );
+
+	HRESULT hResult = InitFunction( Context->TaskBlock->GetInnerBlockForGDKAPI() );
+	if (FAILED(hResult))
+	{
+		delete Context;
+	}
+
+	return hResult;
+}
+
+
+
+HRESULT LocalGDKTask( TFunction<HRESULT(XAsyncBlock*)> InitFunction, TFunction<HRESULT(XAsyncBlock*)> ResultFunction )
+{
+	// sanity check parameters
+	if (InitFunction == nullptr)
+	{
+		return ERROR_INVALID_PARAMETER;
+	}
+
+	// run the local task
+	FGDKLocalTaskBlock Block;
+	HRESULT hResult = InitFunction(Block);
+	if (SUCCEEDED(hResult))
+	{
+		hResult = Block.BlockUntilComplete();
+		if (SUCCEEDED(hResult) && ResultFunction)
+		{
+			hResult = ResultFunction(Block);
+		}
+	}
+
+	return hResult;
+}
+
+
+
+namespace GDKPrivate
+{
+	class FGDKTaskContext
+	{
+	public:
+		FGDKTaskContext();
+		void Init(TUniqueFunction<HRESULT(XAsyncBlock*)> InitFunction);
+		HRESULT Complete(TUniqueFunction<HRESULT(XAsyncBlock*)> ResultFunction);
+		UE::Tasks::FTaskEvent AsyncCompleted;
+
+		HRESULT hResult;
+		TSharedPtr<FGDKAsyncBlock> AsyncBlock;
+	};
+
+
+	FGDKTaskContext::FGDKTaskContext()
+		: AsyncCompleted(UE_SOURCE_LOCATION)
+		, hResult(HRESULT_FROM_WIN32(ERROR_INVALID_FUNCTION))
+	{
+		// prepare the async task block
+		AsyncBlock = MakeShared<FGDKAsyncBlock>( 
+			nullptr, 
+			[this]( FGDKAsyncBlock* TaskBlock )		
+			{
+				AsyncCompleted.Trigger();
+			},
+			FGDKAsyncTaskQueue::GetBackgroundTaskQueue() // always use background task queue: we know the completion callback just sets the trigger, so no need to manage that from the gamethread
+		); 
+	}
+
+	void FGDKTaskContext::Init(TUniqueFunction<HRESULT(XAsyncBlock*)> InitFunction)
+	{
+		// init callback - if it fails then trigger immediately
+		hResult = InitFunction(AsyncBlock->GetInnerBlockForGDKAPI());
+		if (FAILED(hResult))
+		{
+			AsyncCompleted.Trigger();
+		}
+	}
+
+	HRESULT FGDKTaskContext::Complete(TUniqueFunction<HRESULT(XAsyncBlock*)> ResultFunction)
+	{
+		// completion callback - only call if the init callback succeeded
+		if (SUCCEEDED(hResult) && ResultFunction)
+		{
+			hResult = ResultFunction(AsyncBlock->GetInnerBlockForGDKAPI());
+		}
+		return hResult;
+	}
+}
+
+
+FGDKTask LaunchGDKTask(const TCHAR* DebugName, TUniqueFunction<HRESULT(XAsyncBlock*)> InitFunction, TUniqueFunction<HRESULT(XAsyncBlock*)> ResultFunction, UE::Tasks::ETaskPriority Priority)
+{
+	FGDKAsyncTaskMonitor UnusedMonitor;
+	return LaunchGDKTask( DebugName, UnusedMonitor, MoveTemp(InitFunction), MoveTemp(ResultFunction), Priority );
+}
+
+FGDKTask LaunchGDKTask(const TCHAR* DebugName, FGDKAsyncTaskMonitor& OutAsyncTaskMonitor, TUniqueFunction<HRESULT(XAsyncBlock*)> InitFunction, TUniqueFunction<HRESULT(XAsyncBlock*)> ResultFunction, UE::Tasks::ETaskPriority Priority)
+{
+	TUniquePtr<GDKPrivate::FGDKTaskContext> Context = MakeUnique<GDKPrivate::FGDKTaskContext>();
+	Internal_InitAsyncTaskMonitor(OutAsyncTaskMonitor, Context->AsyncBlock);
+
+	// run init
+	Context->Init(MoveTemp(InitFunction));
+	
+	// run result task after init & the context completion
+	return Launch(DebugName,
+		[Context = MoveTemp(Context), ResultFunction = MoveTemp(ResultFunction)]() mutable
+		{
+			return Context->Complete(MoveTemp(ResultFunction));
+		}, 
+		Context->AsyncCompleted,
+		Priority
+	);
+}
+
+
+#endif //WITH_GRDK
diff --git a/Engine/Source/Runtime/Microsoft/GDK/GDKRuntime/Private/GDKThreadCheck.cpp b/Engine/Source/Runtime/Microsoft/GDK/GDKRuntime/Private/GDKThreadCheck.cpp
new file mode 100644
--- /dev/null
+++ b/Engine/Source/Runtime/Microsoft/GDK/GDKRuntime/Private/GDKThreadCheck.cpp
@@ -0,0 +1,39 @@
+// Copyright Epic Games, Inc. All Rights Reserved.
+
+#include "GDKThreadCheck.h"
+
+#if WITH_GDK_THREAD_CHECK
+#include "Misc/Parse.h"
+#include "Misc/CommandLine.h"
+
+#if IS_MONOLITHIC
+#include "Microsoft/AllowMicrosoftPlatformTypes.h"
+#include <XThread.h>
+#include "Microsoft/HideMicrosoftPlatformTypes.h"
+
+FGDKScopedNotTimeSensitive::FGDKScopedNotTimeSensitive()
+{
+	bWasTimeSensitive = XThreadIsTimeSensitive();
+	XThreadSetTimeSensitive(false);
+}
+
+FGDKScopedNotTimeSensitive::~FGDKScopedNotTimeSensitive()
+{
+	XThreadSetTimeSensitive(bWasTimeSensitive);
+}
+#endif //IS_MONOLITHIC
+
+
+bool IsGDKTimeSensitiveThreadCheckEnabled()
+{
+#if WITH_GDK_THREAD_CHECK_ALWAYS
+	return true;
+
+#else
+	static bool bTimeSensitiveThreadCheck = FParse::Param(FCommandLine::Get(), TEXT("TimeSensitiveThreadCheck"));
+	return bTimeSensitiveThreadCheck;
+
+#endif //WITH_GDK_THREAD_CHECK_ALWAYS
+}
+
+#endif //WITH_GDK_THREAD_CHECK
diff --git a/Engine/Source/Runtime/Microsoft/GDK/GDKRuntime/Private/GDKUserManager.cpp b/Engine/Source/Runtime/Microsoft/GDK/GDKRuntime/Private/GDKUserManager.cpp
new file mode 100644
--- /dev/null
+++ b/Engine/Source/Runtime/Microsoft/GDK/GDKRuntime/Private/GDKUserManager.cpp
@@ -0,0 +1,575 @@
+// Copyright Epic Games, Inc. All Rights Reserved.
+
+#include "GDKUserManager.h"
+#if WITH_GRDK
+#include "Misc/MessageDialog.h"
+#include "Misc/CoreDelegates.h"
+#include "HAL/IConsoleManager.h"
+#include "GDKTaskQueueHelpers.h"
+#include "GDKHandle.h"
+#include "GDKThreadCheck.h"
+#include "Misc/ConfigCacheIni.h"
+#include "Misc/App.h"
+#include "Containers/Ticker.h"
+#include "Stats/Stats.h"
+
+#include "Microsoft/AllowMicrosoftPlatformTypes.h"
+THIRD_PARTY_INCLUDES_START
+#include <XGameUI.h>
+#include <XGameErr.h>
+#include <XUser.h>
+THIRD_PARTY_INCLUDES_END
+#include "Microsoft/HideMicrosoftPlatformTypes.h"
+
+DEFINE_LOG_CATEGORY_STATIC(LogGDKUserManager, Log, All);
+
+static int32 GGDKUserAutoDebugDump = 0;
+static FAutoConsoleVariableRef CVarGDKDebugAutoDump(
+	TEXT("GDK.AutoDumpUsers"),
+	GGDKUserAutoDebugDump,
+	TEXT("Logs out before & after states for all user operations.")
+	TEXT("0 = disabled, 1 = after only, 2 = before & after"),
+	ECVF_Default
+);
+
+
+
+
+
+FGDKUserManager::FGDKUserManager()
+{
+	// resize the array of seats to accommodate all potential users
+	uint32_t MaxUsers = 0;
+	HRESULT Result = XUserGetMaxUsers(&MaxUsers);
+	UE_CLOG(FAILED(Result), LogGDKUserManager, Warning, TEXT("FGDKUserManager::Init - XUserGetMaxUsers - Error: (0x%0.8X)."), Result);
+	Seats.AddDefaulted(MaxUsers + 1); //allowing one extra seat to hold all unpaired devices
+	
+	// add callbacks
+	{
+		auto UserChangeFn = [](void* Context, XUserLocalId UserLocalId, XUserChangeEvent Event)
+		{
+			((FGDKUserManager*)Context)->OnUserChange(UserLocalId, Event);
+		};
+		GDK_SCOPE_NOT_TIME_SENSITIVE(); // (startup only) XUserRegisterForChangeEvent is not safe to call on time-sensitive threads
+		XUserRegisterForChangeEvent(FGDKAsyncTaskQueue::GetGenericQueue(), this, UserChangeFn, &UserChangeCallbackToken);
+	}
+
+	FCoreDelegates::ApplicationWillEnterBackgroundDelegate.AddRaw(this, &FGDKUserManager::OnApplicationWillEnterBackground);
+	FCoreDelegates::ApplicationHasEnteredForegroundDelegate.AddRaw(this, &FGDKUserManager::OnApplicationHasEnteredForeground);	
+}
+
+
+
+FGDKUserManager::~FGDKUserManager()
+{
+	GDK_SCOPE_NOT_TIME_SENSITIVE(); // (shutdown only) XUserUnregisterForChangeEvent & XUserUnregisterForDeviceAssociationChanged are not safe to call on time-sensitive threads
+
+	FCoreDelegates::ApplicationWillEnterBackgroundDelegate.RemoveAll(this);
+	FCoreDelegates::ApplicationHasEnteredForegroundDelegate.RemoveAll(this);
+
+	XUserUnregisterForChangeEvent(UserChangeCallbackToken, true);
+
+	{
+		FScopeLock Lock(&StateLock);
+		Seats.Reset();
+	}
+}
+
+
+
+bool FGDKUserManager::HasDefaultUser() const
+{
+	// default user will always be user 0
+	return Seats[0].PairedUserHandle.IsValid();
+}
+
+void FGDKUserManager::PickUserAsync(XUserAddOptions Options, FGDKPickUserCompleteDelegate OnPickUserCompleteDelegate)
+{
+	HRESULT hResult = AsyncGDKTask(
+		[Options](XAsyncBlock* Block)
+		{
+			HRESULT hResult = XUserAddAsync(Options, Block);
+			UE_CLOG(FAILED(hResult), LogGDKUserManager, Warning, TEXT("XUserAddAsync failed: 0x%0.8X"), hResult);
+			return hResult;
+		},
+		[this, Options, OnPickUserCompleteDelegate] (XAsyncBlock* Block)
+		{
+			FGDKUserHandle User;
+			HRESULT hResult = XUserAddResult(Block, User.GetInitReference());
+			UE_CLOG(FAILED(hResult) && (hResult != E_ACCESSDENIED), LogGDKUserManager, Warning, TEXT("XUserAddResult failed: 0x%0.8X"), hResult);
+			HandleXUserAddComplete(Options, hResult, User, OnPickUserCompleteDelegate);
+		}
+	);
+
+	if (FAILED(hResult))
+	{
+		PickUserAsyncComplete(hResult, FGDKUserHandle(), OnPickUserCompleteDelegate);
+	}
+}
+
+void FGDKUserManager::HandleXUserAddComplete(XUserAddOptions Options, HRESULT hResult, FGDKUserHandle User, FGDKPickUserCompleteDelegate OnPickUserCompleteDelegate)
+{
+	if (hResult == E_ACCESSDENIED)
+	{
+		PickUserAsync(Options, OnPickUserCompleteDelegate);
+		return;
+	}
+
+	if (SUCCEEDED(hResult))
+	{
+		uint64 UserId;
+		hResult = XUserGetId(User, &UserId);
+		UE_CLOG(FAILED(hResult), LogGDKUserManager, Warning, TEXT("XUserGetId failed: 0x%0.8X"), hResult);
+	}
+
+	if (hResult == E_GAMEUSER_RESOLVE_USER_ISSUE_REQUIRED && !FApp::IsUnattended())
+	{
+		ResolveUserIssueWithUiAsync(User, OnPickUserCompleteDelegate);
+		return;
+	}
+
+	PickUserAsyncComplete(hResult, User, OnPickUserCompleteDelegate);
+}
+
+void FGDKUserManager::ResolveUserIssueWithUiAsync(FGDKUserHandle User, FGDKPickUserCompleteDelegate OnPickUserCompleteDelegate)
+{
+	HRESULT hResult = AsyncGDKTask(
+		[User](XAsyncBlock* Block)
+		{
+			HRESULT hResult = XUserResolveIssueWithUiAsync(User, nullptr, Block);
+			UE_CLOG(FAILED(hResult), LogGDKUserManager, Warning, TEXT("XUserResolveIssueWithUiAsync failed: 0x%0.8X"), hResult);
+			return hResult;
+		},
+		[this, OnPickUserCompleteDelegate, User] (XAsyncBlock* Block)
+		{
+			HRESULT hResult = XUserResolveIssueWithUiResult(Block);
+			UE_CLOG(FAILED(hResult) && (hResult != E_ACCESSDENIED), LogGDKUserManager, Warning, TEXT("XUserResolveIssueWithUiResult failed: 0x%0.8X"), hResult);
+			HandleResolveUserIssueWithUiComplete(hResult, User, OnPickUserCompleteDelegate);
+		}
+	);
+
+	if (FAILED(hResult))
+	{
+		PickUserAsyncComplete(hResult, User, OnPickUserCompleteDelegate);
+	}
+}
+
+void FGDKUserManager::HandleResolveUserIssueWithUiComplete(HRESULT hResult, FGDKUserHandle User, FGDKPickUserCompleteDelegate OnPickUserCompleteDelegate)
+{
+	if (hResult == E_ACCESSDENIED)
+	{
+		ResolveUserIssueWithUiAsync(User, OnPickUserCompleteDelegate);
+		return;
+	}
+
+	PickUserAsyncComplete(hResult, User, OnPickUserCompleteDelegate);
+}
+
+void FGDKUserManager::PickUserAsyncComplete(HRESULT hResult, FGDKUserHandle User, FGDKPickUserCompleteDelegate OnPickUserCompleteDelegate)
+{
+	if (User.IsValid() && GetCurrentSeatForUser(User) == INDEX_NONE)
+	{
+		XUserLocalId UserLocalId;
+		if (SUCCEEDED(XUserGetLocalId(User, &UserLocalId)))
+		{
+			FScopeLock Lock(&StateLock);
+			SignedOutXUserLocalIds.Remove(UserLocalId.value);
+		}
+
+		AutoDebugDump(TEXT("Before (via PickUser Callback)"), EDebugDumpType::BeforeStateChange);
+			
+		int32 SeatIndex = FindSeatForUser(User);
+		check(SeatIndex != INDEX_NONE);
+			
+		UE_LOG(LogGDKUserManager, Log, TEXT("New user %s[%d] signed in via callback"), *LexToString(User), SeatIndex);
+		AutoDebugDump(TEXT("After (via PickUser Callback)"), EDebugDumpType::AfterStateChange);
+	}
+
+	OnPickUserCompleteDelegate.ExecuteIfBound(hResult, User);
+}
+
+void FGDKUserManager::AddDefaultUser()
+{
+	// block and try to silently add the default user
+	FGDKUserHandle User;
+	HRESULT hResult = LocalGDKTask( 
+		[&User]( XAsyncBlock* Block )
+		{
+			return XUserAddAsync( XUserAddOptions::AddDefaultUserSilently, Block );
+		},
+		[&User]( XAsyncBlock* Block )
+		{
+			return XUserAddResult(Block, User.GetInitReference());
+		}
+	);
+
+	// check initial result
+	if (SUCCEEDED(hResult))
+	{
+		uint64 UserId;
+		hResult = XUserGetId(User, &UserId);
+		UE_CLOG(FAILED(hResult), LogGDKUserManager, Warning, TEXT("Failed to get xuid for default user"));
+	}
+
+	// check final result
+	if (SUCCEEDED(hResult))
+	{
+		ensure(FindSeatForUser(User) != INDEX_NONE);
+	}
+	else if (hResult == E_GAMEUSER_RESOLVE_USER_ISSUE_REQUIRED && !FApp::IsUnattended())
+	{
+		// cannot block & show resolve UI, so start an async task instead
+		ResolveUserIssueWithUiAsync(User, FGDKPickUserCompleteDelegate::CreateLambda([this]
+			(HRESULT hResult, FGDKUserHandle User)
+			{
+				ensure(FindSeatForUser(User) != INDEX_NONE);
+			}
+		));
+	}
+	else
+	{
+		//NB. if you see E_ABORT here, and the Xbox live login dialog box gave error code 0x87DD0005 it probably means this Device Family isn't enabled for Xbox Live in Partner Center -> Xbox Live -> Xbox Live Settings
+		UE_LOG(LogGDKUserManager, Warning, TEXT("Failed to add default user. Error 0x%X"), hResult);
+		UE_CLOG((hResult == E_GAMEUSER_NO_DEFAULT_USER),             LogGDKUserManager, Warning, TEXT("No default user."));
+		UE_CLOG((hResult == E_GAMEUSER_NO_PACKAGE_IDENTITY),         LogGDKUserManager, Warning, TEXT("No package identity"));
+		UE_CLOG((hResult == E_GAMEUSER_RESOLVE_USER_ISSUE_REQUIRED), LogGDKUserManager, Warning, TEXT("Not allowed to display user resolve UI.") );
+	}
+}
+
+
+
+
+
+FGDKUserHandle FGDKUserManager::GetUserHandleByPlatformId(FPlatformUserId PlatformUserId) const
+{
+	int32 SeatIndex = PlatformUserId;
+	if (!Seats.IsValidIndex(SeatIndex))
+	{
+		return FGDKUserHandle();
+	}
+
+	return Seats[SeatIndex].PairedUserHandle;
+}
+
+
+FPlatformUserId FGDKUserManager::GetPlatformIdForSeat(int32 SeatIndex) const
+{
+	if (Seats.IsValidIndex(SeatIndex))
+	{
+		return FPlatformMisc::GetPlatformUserForUserIndex(SeatIndex);
+	}
+	
+	return PLATFORMUSERID_NONE;
+}
+
+
+FGDKUserHandle FGDKUserManager::GetUserHandleByXUserId(uint64 XUserId) const
+{
+	GDK_SCOPE_NOT_TIME_SENSITIVE(); // XUserFindUserById is not safe to call on a time-sensitive thread
+
+	FGDKUserHandle UserHandle;
+	XUserFindUserById(XUserId, UserHandle.GetInitReference());
+	return UserHandle;
+}
+
+
+TOptional<uint64> FGDKUserManager::GetXUserIdByPlatformId(FPlatformUserId PlatformUserId) const
+{
+	FGDKUserHandle UserHandle = GetUserHandleByPlatformId(PlatformUserId);
+
+	TOptional<uint64> Result;
+
+	if (UserHandle.IsValid())
+	{
+		uint64 Xuid;
+		if (SUCCEEDED(XUserGetId(UserHandle, &Xuid)))
+		{
+			Result = Xuid;
+		}
+	}
+
+	return Result;
+}
+
+
+FPlatformUserId FGDKUserManager::GetPlatformIdByUserHandle(const FGDKUserHandle& UserHandle)
+{
+	return GetPlatformIdForSeat(FindSeatForUser(UserHandle));
+}
+
+
+
+int32 FGDKUserManager::FindSeatForUser( const FGDKUserHandle& Handle)
+{
+	FScopeLock Lock(&StateLock);
+
+	if (!Handle.IsValid())
+	{
+		return INDEX_NONE;
+	}
+
+	int32 SeatIndex = GetCurrentSeatForUser(Handle);
+	if (SeatIndex == INDEX_NONE)
+	{
+		XUserLocalId UserLocalId;
+		if (SUCCEEDED(XUserGetLocalId(Handle, &UserLocalId)) && SignedOutXUserLocalIds.Contains(UserLocalId.value))
+		{
+			UE_LOG(LogGDKUserManager, Log, TEXT("Will not assign seat to signed-out user %s"), *LexToString(Handle) );
+			return INDEX_NONE;
+		}
+
+		AutoDebugDump(TEXT("Before AssignSeatForUser"), EDebugDumpType::BeforeStateChange);
+
+		// we have not heard about this user before so assign them a seat now
+		SeatIndex = FindEmptySeat();
+		check(SeatIndex != INDEX_NONE);
+		Seats[SeatIndex].PairedUserHandle = Handle;
+
+		UE_LOG(LogGDKUserManager, Log, TEXT("Assigning user %s to seat %d"), *LexToString(Handle), SeatIndex);
+		AutoDebugDump(TEXT("After AssignSeatForUser"), EDebugDumpType::AfterStateChange);
+
+		BroadcastUserAddedOnGameThread(SeatIndex);
+	}
+
+	return SeatIndex;
+}
+
+
+
+int32 FGDKUserManager::GetCurrentSeatForUser(const FGDKUserHandle& Handle) const
+{
+	if (!Handle.IsValid())
+	{
+		return INDEX_NONE;
+	}
+
+	return Seats.IndexOfByPredicate([Handle](const FSeat& Seat)
+	{
+		return (Seat.PairedUserHandle.IsValid() && Handle == Seat.PairedUserHandle);
+	});
+}
+
+
+
+int32 FGDKUserManager::FindEmptySeat() const
+{
+	for( int SeatIndex = 0; SeatIndex < Seats.Num(); SeatIndex++ )
+	{
+		if (Seats[SeatIndex].PairedUserHandle.IsValid())
+		{
+			continue;
+		}
+
+		return SeatIndex;
+	}
+
+	return INDEX_NONE;
+}
+
+
+TArray<FGDKUserHandle> FGDKUserManager::GetAllUserHandles() const
+{
+	TArray<FGDKUserHandle> Result;
+	for( int SeatIndex = 0; SeatIndex < Seats.Num(); SeatIndex++ )
+	{
+		if (Seats[SeatIndex].PairedUserHandle.IsValid())
+		{
+			Result.Add( Seats[SeatIndex].PairedUserHandle );
+		}
+	}
+
+	return MoveTemp(Result);
+}
+
+
+bool FGDKUserManager::WasUserRecentlySignedOut(XUserLocalId UserLocalId) const
+{
+	return SignedOutXUserLocalIds.Contains(UserLocalId.value);
+}
+
+
+void FGDKUserManager::OnUserChange(XUserLocalId UserLocalId, XUserChangeEvent Event)
+{
+	GDK_SCOPE_NOT_TIME_SENSITIVE(); // XUserFindUserByLocalId is not safe to call on a time-sensitive thread
+
+	FGDKUserHandle UserHandle;
+	XUserFindUserByLocalId(UserLocalId, UserHandle.GetInitReference());
+
+	if (Event == XUserChangeEvent::SignedInAgain )
+	{
+		// do nothing here: seat assignment & signout state should handled in PickUserAsyncComplete for consistency with genuine new user flow
+	}
+	else if (Event == XUserChangeEvent::SigningOut)
+	{
+		{
+			FScopeLock Lock(&StateLock);
+			SignedOutXUserLocalIds.Add(UserLocalId.value);
+		}
+
+		int32 SeatIndex = GetCurrentSeatForUser(UserHandle);
+		if (SeatIndex != INDEX_NONE)
+		{
+			UE_LOG(LogGDKUserManager, Log, TEXT("User %s[%d] signing out via callback"), *LexToString(UserHandle), SeatIndex);
+			RemoveUserOnGameThread(SeatIndex);
+		}
+	}
+}
+
+
+void FGDKUserManager::BroadcastUserAddedOnGameThread(int32 SeatIndex)
+{
+	if (IsInGameThread())
+	{
+		// we are on the gamethread - signal the delegate directly
+		FCoreDelegates::OnUserLoginChangedEvent.Broadcast(true, GetPlatformIdForSeat(SeatIndex), SeatIndex);
+		UE_CLOG( GGDKUserAutoDebugDump != 0, LogGDKUserManager, Log, TEXT("\n**** FCoreDelegates::OnUserLoginChangedEvent - User %d Login ****\n"), SeatIndex );
+	}
+	else 
+	{
+		// we are not on the gamethread - request a callback when we are
+		ExecuteOnGameThread(UE_SOURCE_LOCATION, [SeatIndex, this]()
+		{
+			BroadcastUserAddedOnGameThread(SeatIndex);
+		});
+	}
+
+}
+
+void FGDKUserManager::RemoveUserOnGameThread(int32 SeatIndex)
+{
+	if (IsInGameThread())
+	{
+		// we are on the gamethread - remove the user
+		AutoDebugDump(TEXT("Before User Remove"), EDebugDumpType::BeforeStateChange);
+		UE_LOG(LogGDKUserManager, Log, TEXT("User %s (%d) signed out"), *LexToString(Seats[SeatIndex].PairedUserHandle), SeatIndex);
+
+		FCoreDelegates::OnUserLoginChangedEvent.Broadcast(false, GetPlatformIdForSeat(SeatIndex), SeatIndex);
+		UE_CLOG( GGDKUserAutoDebugDump != 0, LogGDKUserManager, Log, TEXT("\n**** FCoreDelegates::OnUserLoginChangedEvent - User %d Logout ****\n"), SeatIndex );
+
+		{
+			FScopeLock Lock(&StateLock);
+			Seats[SeatIndex].PairedUserHandle.Clear();
+		}
+
+		AutoDebugDump(TEXT("After User Remove"), EDebugDumpType::AfterStateChange);
+	}
+	else
+	{
+		// we are not on the gamethread - request a callback when we are
+		ExecuteOnGameThread(UE_SOURCE_LOCATION, [SeatIndex, this]()
+		{
+			RemoveUserOnGameThread(SeatIndex);
+		});
+	}
+}
+
+
+void FGDKUserManager::OnApplicationWillEnterBackground()
+{
+	SCOPED_ENTER_BACKGROUND_EVENT(STAT_FGDKUserManager_OnApplicationWillEnterBackground);
+	AutoDebugDump(TEXT("OnApplicationWillEnterBackground"), EDebugDumpType::AfterStateChange);
+}
+
+void FGDKUserManager::OnApplicationHasEnteredForeground()
+{
+	AutoDebugDump(TEXT("OnApplicationHasEnteredForeground"), EDebugDumpType::AfterStateChange);
+}
+
+
+
+
+void FGDKUserManager::DebugDump( const TCHAR* BannerMessage, EDebugDumpType::Type Type ) const
+{
+	GDK_SCOPE_NOT_TIME_SENSITIVE(); // (debug only) XUserFindForDevice isn't safe to call on time-sensitive threads
+
+	// write banner
+	UE_LOG( LogGDKUserManager, Log, TEXT("*** GDK User Manager : %s ***"), BannerMessage ? BannerMessage : TEXT("") );
+
+	// helper lambda to write out empty spans
+	int EmptyRangeStart = INDEX_NONE;
+	auto WriteEmptyItems = [&EmptyRangeStart]( int SeatIndex )
+	{
+		if (EmptyRangeStart != INDEX_NONE)
+		{
+			if (EmptyRangeStart == SeatIndex - 1)
+			{
+				UE_LOG(LogGDKUserManager, Log, TEXT("\t\t%d - empty"), EmptyRangeStart);
+			}
+			else
+			{
+				UE_LOG(LogGDKUserManager, Log, TEXT("\t\t%d-%d : empty"), EmptyRangeStart, SeatIndex - 1);
+			}
+			EmptyRangeStart = INDEX_NONE;
+		}
+	};
+
+	// write out all users
+	for( int SeatIndex = 0; SeatIndex < Seats.Num(); SeatIndex++ )
+	{
+		const FSeat& Seat = Seats[SeatIndex];
+		bool bIsEmpty = !Seat.PairedUserHandle.IsValid();
+
+		bool bWasEmpty = (EmptyRangeStart != INDEX_NONE);
+		if (bWasEmpty != bIsEmpty)
+		{
+			if (bWasEmpty)
+			{
+				WriteEmptyItems(SeatIndex);
+			}
+			else
+			{
+				EmptyRangeStart = SeatIndex;
+			}
+		}
+		if (!bIsEmpty)
+		{
+			UE_LOG(LogGDKUserManager, Log, TEXT("\t\t%d : %s"), SeatIndex, *LexToString(Seat.PairedUserHandle));
+		}
+	}
+	WriteEmptyItems(Seats.Num());
+
+	// check users
+	for( int SeatIndex = 0; SeatIndex < Seats.Num(); SeatIndex++ )
+	{
+		const FSeat& Seat = Seats[SeatIndex];
+		if (Seat.PairedUserHandle.IsValid())
+		{
+			XUserState State;
+			HRESULT Result = XUserGetState(Seat.PairedUserHandle, &State);
+			if (SUCCEEDED(Result))
+			{
+				UE_CLOG(State != XUserState::SignedIn, LogGDKUserManager, Warning, TEXT("User in seat %d is not signed in (%d)"), SeatIndex, State);
+			}
+			else if (Result == E_GAMERUNTIME_SUSPENDED)
+			{
+				//... do nothing - cannot query user state while suspended
+			}
+			else
+			{
+				UE_LOG(LogGDKUserManager, Warning, TEXT("Failed to get user state for user in seat %d : HR=0x%X"), SeatIndex, Result );
+			}
+		}
+	}
+
+	// done
+	UE_LOG( LogGDKUserManager, Log, TEXT("****************************************************\n") );
+
+}
+
+ void FGDKUserManager::DebugDump() const
+ {
+	 DebugDump( TEXT("Debug Dump") );
+ }
+
+void FGDKUserManager::AutoDebugDump( const TCHAR* BannerMessage, EDebugDumpType::Type Type )
+{
+	bool bAfterStateChange = (Type & EDebugDumpType::AfterStateChange) != 0;
+
+	if (GGDKUserAutoDebugDump && ( GGDKUserAutoDebugDump == 2 || bAfterStateChange ) ) // 0 = disabled, 1 = after only, 2 = both
+	{
+		DebugDump(BannerMessage, Type );
+	}
+}
+
+
+#endif //WITH_GRDK
diff --git a/Engine/Source/Runtime/Microsoft/GDK/GDKRuntime/Private/GDKUserManager.h b/Engine/Source/Runtime/Microsoft/GDK/GDKRuntime/Private/GDKUserManager.h
new file mode 100644
--- /dev/null
+++ b/Engine/Source/Runtime/Microsoft/GDK/GDKRuntime/Private/GDKUserManager.h
@@ -0,0 +1,85 @@
+// Copyright Epic Games, Inc. All Rights Reserved.
+
+#pragma once
+
+#include "GDKRuntimeModule.h"
+#if WITH_GRDK
+#include "CoreMinimal.h"
+#include "GDKHandle.h"
+#include "HAL/CriticalSection.h"
+#include "Containers/Set.h"
+#include "Misc/Optional.h"
+
+#include "Microsoft/AllowMicrosoftPlatformTypes.h"
+#include <XUser.h>
+#include "Microsoft/HideMicrosoftPlatformTypes.h"
+
+#define UE_API GDKRUNTIME_API
+
+/*
+ * The GDK user manager is responsible for assigning a suitable FPlatformUserId for GDK users. FPlatformUserIds are currently recycled when a user signs out & is replaced by another and represent the seat that the user is assigned to
+ */
+class FGDKUserManager
+{
+public:
+	UE_API FGDKUserManager();
+	UE_API virtual ~FGDKUserManager();
+
+	UE_API FGDKUserHandle GetUserHandleByPlatformId(FPlatformUserId PlatformUserId) const;
+	UE_API FGDKUserHandle GetUserHandleByXUserId(uint64 XUserId) const;
+	UE_API FPlatformUserId GetPlatformIdByUserHandle(const FGDKUserHandle& UserHandle);
+	UE_API TOptional<uint64> GetXUserIdByPlatformId(FPlatformUserId PlatformUserId) const;
+	UE_API TArray<FGDKUserHandle> GetAllUserHandles() const;
+	UE_API bool WasUserRecentlySignedOut(XUserLocalId UserLocalId) const;
+	UE_API int32 FindEmptySeat() const;
+	UE_API bool HasDefaultUser() const;
+	UE_API void AddDefaultUser();
+	UE_API void PickUserAsync(XUserAddOptions Options, FGDKPickUserCompleteDelegate OnPickUserCompleteDelegate);
+
+	UE_API void DebugDump() const;
+
+protected:
+	// internal helper functions
+	UE_API FPlatformUserId GetPlatformIdForSeat(int32 SeatIndex) const;
+	UE_API int32 FindSeatForUser(const FGDKUserHandle& UserHandle);
+	UE_API int32 GetCurrentSeatForUser(const FGDKUserHandle& UserHandle) const;
+
+private:
+	// callbacks
+	void OnUserChange(XUserLocalId UserLocalId, XUserChangeEvent Event);
+	void OnApplicationWillEnterBackground();
+	void OnApplicationHasEnteredForeground();
+
+	// helper functions
+	FORCEINLINE static bool IsValid(const XUserLocalId& UserLocalId) { return UserLocalId.value != 0; }
+
+	void HandleXUserAddComplete(XUserAddOptions Options, HRESULT hResult, FGDKUserHandle User, FGDKPickUserCompleteDelegate OnPickUserCompleteDelegate);
+	void ResolveUserIssueWithUiAsync(FGDKUserHandle User, FGDKPickUserCompleteDelegate OnResolveUserIssueWithUiComplete);
+	void HandleResolveUserIssueWithUiComplete(HRESULT hResult, FGDKUserHandle User, FGDKPickUserCompleteDelegate OnPickUserCompleteDelegate);
+	void PickUserAsyncComplete(HRESULT hResult, FGDKUserHandle User, FGDKPickUserCompleteDelegate OnPickUserCompleteDelegate);
+
+	struct EDebugDumpType
+	{
+		typedef uint8 Type;
+		static const uint8 BeforeStateChange = 1;
+		static const uint8 AfterStateChange = 2;
+	};
+	void AutoDebugDump( const TCHAR* BannerMessage, EDebugDumpType::Type Type );
+	void DebugDump( const TCHAR* BannerMessage, EDebugDumpType::Type Type = 0 ) const;
+
+	void BroadcastUserAddedOnGameThread( int32 SeatIndex );
+	void RemoveUserOnGameThread( int32 SeatIndex );
+
+	struct FSeat
+	{
+		FGDKUserHandle PairedUserHandle;
+	};
+	TArray<FSeat>               Seats;
+	TSet<uint64>                SignedOutXUserLocalIds;
+	FCriticalSection            StateLock;
+	XTaskQueueRegistrationToken UserChangeCallbackToken;
+};
+
+#undef UE_API
+
+#endif //WITH_GRDK
diff --git a/Engine/Source/Runtime/Microsoft/GDK/GDKRuntime/Private/Tests/GDKRuntimeTests.cpp b/Engine/Source/Runtime/Microsoft/GDK/GDKRuntime/Private/Tests/GDKRuntimeTests.cpp
new file mode 100644
--- /dev/null
+++ b/Engine/Source/Runtime/Microsoft/GDK/GDKRuntime/Private/Tests/GDKRuntimeTests.cpp
@@ -0,0 +1,236 @@
+// Copyright Epic Games, Inc. All Rights Reserved.
+
+#include "GDKRuntimeModule.h"
+#include "Misc/AutomationTest.h"
+#if WITH_GRDK && WITH_DEV_AUTOMATION_TESTS
+#include "GDKTaskQueueHelpers.h"
+#include "HAL/PlatformProcess.h"
+#include "HAL/Event.h"
+
+#include "Microsoft/AllowMicrosoftPlatformTypes.h"
+THIRD_PARTY_INCLUDES_START
+#include <XGameUI.h>
+THIRD_PARTY_INCLUDES_END
+#include "Microsoft/HideMicrosoftPlatformTypes.h"
+
+// MyGame -ExecCmds="Automation RunTests GDKTests" -TestExit="Automation Test Queue Empty" -unattended
+namespace GDKTests
+{
+	struct FTestContext
+	{
+		FTestContext(FAutomationTestBase* InTest) : Test(InTest) {}
+		FAutomationTestBase* Test;
+	};
+
+
+	DEFINE_LATENT_AUTOMATION_COMMAND_ONE_PARAMETER(FGDKTestCmd_Wait, float, Duration); // NB. not using FWaitLatentCommand because it's in Engine and we're in Core
+	bool FGDKTestCmd_Wait::Update()
+	{
+		const double NewTime = FPlatformTime::Seconds();
+		if (NewTime - StartTime >= Duration)
+		{
+			return true;
+		}
+		return false;
+	}
+
+	DEFINE_LATENT_AUTOMATION_COMMAND_TWO_PARAMETER( FGDKTestCmd_WaitForEvent, FAutomationTestBase*, Test, FEvent*, Event );
+	bool FGDKTestCmd_WaitForEvent::Update()
+	{
+		return Event->Wait(FTimespan::FromSeconds(0.1f));
+		
+	}
+
+	DEFINE_LATENT_AUTOMATION_COMMAND_THREE_PARAMETER( FGDKTestCmd_CheckResult, FAutomationTestBase*, Test, HRESULT, hResult, HRESULT, hResultExpected );
+	bool FGDKTestCmd_CheckResult::Update()
+	{
+		if (hResultExpected != hResult)
+		{
+			Test->AddError(FString::Printf(TEXT("Expected hResult 0x%X but got 0x%X"), hResultExpected, hResult ));
+		}
+		return true;
+	}
+
+	DEFINE_LATENT_AUTOMATION_COMMAND_ONE_PARAMETER( FGDKTestCmd_RunOp, TFunction<void(void)>, Op );
+	bool FGDKTestCmd_RunOp::Update()
+	{
+		Op();
+		return true;
+	}
+
+
+
+	struct FTaskQueueTestContext : public FTestContext
+	{
+		FTaskQueueTestContext(FAutomationTestBase* InTest) 
+			: FTestContext(InTest)
+			, Event(FPlatformProcess::GetSynchEventFromPool(false))
+			, hResult(E_NOTIMPL)
+		{
+		}
+
+		~FTaskQueueTestContext()
+		{
+			FPlatformProcess::ReturnSynchEventToPool(Event);
+			TaskMonitor.TryCancel(false);
+		}
+
+		// NB. some redudent members depending on the test
+		FEvent* Event;
+		HRESULT hResult;
+		FGDKAsyncTaskMonitor TaskMonitor;
+		FGDKTask Task;
+	};
+
+
+
+	DEFINE_LATENT_AUTOMATION_COMMAND_ONE_PARAMETER( FGDKTestCmd_TaskQueue_Cancel, FTaskQueueTestContext*, TestContext );
+	bool FGDKTestCmd_TaskQueue_Cancel::Update()
+	{
+		bool bCancelResult = TestContext->TaskMonitor.TryCancel(true);
+		TestContext->Test->TestTrue("task cancel success", bCancelResult);
+		return true;
+	}
+
+	DEFINE_LATENT_AUTOMATION_COMMAND_ONE_PARAMETER( FGDKTestCmd_TaskQueue_WaitForTask, FTaskQueueTestContext*, TestContext );
+	bool FGDKTestCmd_TaskQueue_WaitForTask::Update()
+	{
+		return TestContext->Task.Wait(FTimespan::FromSeconds(0.1f));
+	}
+
+
+	DEFINE_LATENT_AUTOMATION_COMMAND_TWO_PARAMETER( FGDKTestCmd_TaskQueue_CheckResult, FTaskQueueTestContext*, TestContext, HRESULT, hResultExpected );
+	bool FGDKTestCmd_TaskQueue_CheckResult::Update()
+	{
+		HRESULT hResult = TestContext->Task.GetResult();
+		if (hResultExpected != hResult)
+		{
+			TestContext->Test->AddError(FString::Printf(TEXT("Expected hResult 0x%X but got 0x%X"), hResultExpected, hResult ));
+		}
+		return true;
+	}
+
+	DEFINE_LATENT_AUTOMATION_COMMAND_ONE_PARAMETER( FGDKTestCmd_TaskQueue_CleanUp, FTaskQueueTestContext*, TestContext );
+	bool FGDKTestCmd_TaskQueue_CleanUp::Update()
+	{
+		delete TestContext;
+		return true;
+	}
+
+
+	// run an async task and then attempt to cancel it from the game thread. verifies the completion callback fires and returns E_ABORT
+	IMPLEMENT_SIMPLE_AUTOMATION_TEST(FGDKTaskQueueTestCancel, "GDKTests.TaskQueue.CancelAsyncGDKTask", EAutomationTestFlags::ClientContext | EAutomationTestFlags::EngineFilter )
+	bool FGDKTaskQueueTestCancel::RunTest(const FString& Parameters)
+	{
+		FTaskQueueTestContext* TestContext = new FTaskQueueTestContext(this);
+
+		auto Op = [TestContext]()
+		{
+			AsyncGDKTask( 
+				TestContext->TaskMonitor,
+				[TestContext]( XAsyncBlock* Block )
+				{
+					return XGameUiShowMessageDialogAsync( Block, TCHAR_TO_UTF8(*TestContext->Test->GetTestFullName()), "Should Be Cancelled Automatically", "Test", "Test", "Test", XGameUiMessageDialogButton::First, XGameUiMessageDialogButton::Second );
+				},
+				[TestContext]( XAsyncBlock* Block )
+				{
+					XGameUiMessageDialogButton Button;
+					TestContext->hResult = XGameUiShowMessageDialogResult(Block, &Button);
+					TestContext->Event->Trigger();
+				}
+			);
+		};
+
+		ADD_LATENT_AUTOMATION_COMMAND( FGDKTestCmd_RunOp(Op));
+		ADD_LATENT_AUTOMATION_COMMAND( FGDKTestCmd_Wait(1.0f) );
+		ADD_LATENT_AUTOMATION_COMMAND( FGDKTestCmd_TaskQueue_Cancel(TestContext) );
+		ADD_LATENT_AUTOMATION_COMMAND( FGDKTestCmd_WaitForEvent(this, TestContext->Event) );
+		ADD_LATENT_AUTOMATION_COMMAND( FGDKTestCmd_CheckResult(this, TestContext->hResult, E_ABORT) );
+		ADD_LATENT_AUTOMATION_COMMAND( FGDKTestCmd_TaskQueue_CleanUp(TestContext) );
+
+		return true;
+	}
+
+
+	// run an async task and then attempt to cancel it in the background. verifies the completion callback fires and returns E_ABORT
+	IMPLEMENT_SIMPLE_AUTOMATION_TEST(FGDKTaskQueueTestCancelFromBgThread, "GDKTests.TaskQueue.BgCancelAsyncGDKTask", EAutomationTestFlags::ClientContext | EAutomationTestFlags::EngineFilter )
+	bool FGDKTaskQueueTestCancelFromBgThread::RunTest(const FString& Parameters)
+	{
+		FTaskQueueTestContext* TestContext = new FTaskQueueTestContext(this);
+
+		auto Op = [TestContext]()
+		{
+			AsyncGDKTask( 
+				TestContext->TaskMonitor,
+				[TestContext]( XAsyncBlock* Block )
+				{
+					return XGameUiShowMessageDialogAsync( Block, TCHAR_TO_UTF8(*TestContext->Test->GetTestFullName()), "Should Be Cancelled Automatically", "Test", "Test", "Test", XGameUiMessageDialogButton::First, XGameUiMessageDialogButton::Second );
+				},
+				[TestContext]( XAsyncBlock* Block )
+				{
+					XGameUiMessageDialogButton Button;
+					TestContext->hResult = XGameUiShowMessageDialogResult(Block, &Button);
+					TestContext->Event->Trigger();
+				}
+			);
+		};
+
+		auto CancelOpOnBgThread = [TestContext]()
+		{
+			AsyncTask( ENamedThreads::AnyNormalThreadNormalTask,[TestContext]()
+			{
+				check(!IsInGameThread());
+				bool bCancelResult = TestContext->TaskMonitor.TryCancel(true);
+				TestContext->Test->TestEqual("Check cancel success from background thread", bCancelResult, true);
+			});
+		};
+
+
+		ADD_LATENT_AUTOMATION_COMMAND( FGDKTestCmd_RunOp(Op));
+		ADD_LATENT_AUTOMATION_COMMAND( FGDKTestCmd_Wait(1.0f) );
+		ADD_LATENT_AUTOMATION_COMMAND( FGDKTestCmd_RunOp(CancelOpOnBgThread) );
+		ADD_LATENT_AUTOMATION_COMMAND( FGDKTestCmd_WaitForEvent(this, TestContext->Event) );
+		ADD_LATENT_AUTOMATION_COMMAND( FGDKTestCmd_CheckResult(this, TestContext->hResult, E_ABORT) );
+		ADD_LATENT_AUTOMATION_COMMAND( FGDKTestCmd_TaskQueue_CleanUp(TestContext) );
+
+		return true;
+	}
+
+
+
+	// create an async task using the UE::Task system and then attempt to cancel it. verifies the completion callback fires and returns E_ABORT
+	IMPLEMENT_SIMPLE_AUTOMATION_TEST(FGDKTaskQueueTestCancelTask, "GDKTests.TaskQueue.CancelLaunchGDKTask", EAutomationTestFlags::ClientContext | EAutomationTestFlags::EngineFilter )
+	bool FGDKTaskQueueTestCancelTask::RunTest(const FString& Parameters)
+	{
+		FTaskQueueTestContext* TestContext = new FTaskQueueTestContext(this);
+
+		auto Op = [TestContext]()
+		{
+			TestContext->Task = LaunchGDKTask( 
+				UE_SOURCE_LOCATION,
+				TestContext->TaskMonitor,
+				[TestContext]( XAsyncBlock* Block )
+				{
+					return XGameUiShowMessageDialogAsync( Block, TCHAR_TO_UTF8(*TestContext->Test->GetTestFullName()), "Should Be Cancelled Automatically", "Test", "Test", "Test", XGameUiMessageDialogButton::First, XGameUiMessageDialogButton::Second );
+				},
+				[TestContext]( XAsyncBlock* Block )
+				{
+					XGameUiMessageDialogButton Button;
+					return XGameUiShowMessageDialogResult(Block, &Button);
+				}
+			);
+		};
+
+		ADD_LATENT_AUTOMATION_COMMAND( FGDKTestCmd_RunOp(Op));
+		ADD_LATENT_AUTOMATION_COMMAND( FGDKTestCmd_Wait(1.0f) );
+		ADD_LATENT_AUTOMATION_COMMAND( FGDKTestCmd_TaskQueue_Cancel(TestContext) );
+		ADD_LATENT_AUTOMATION_COMMAND( FGDKTestCmd_TaskQueue_WaitForTask(TestContext) );
+		ADD_LATENT_AUTOMATION_COMMAND( FGDKTestCmd_TaskQueue_CheckResult(TestContext, E_ABORT) );
+		ADD_LATENT_AUTOMATION_COMMAND( FGDKTestCmd_TaskQueue_CleanUp(TestContext) );
+
+		return true;
+	}
+}
+
+#endif //WITH_GRDK && WITH_DEV_AUTOMATION_TESTS
+
diff --git a/Engine/Source/Runtime/Microsoft/GDK/GDKRuntime/Public/GDKHandle.h b/Engine/Source/Runtime/Microsoft/GDK/GDKRuntime/Public/GDKHandle.h
new file mode 100644
--- /dev/null
+++ b/Engine/Source/Runtime/Microsoft/GDK/GDKRuntime/Public/GDKHandle.h
@@ -0,0 +1,176 @@
+// Copyright Epic Games, Inc. All Rights Reserved.
+#pragma once
+
+#if WITH_GRDK
+
+#include "Misc/AssertionMacros.h"
+#include "Microsoft/AllowMicrosoftPlatformTypes.h"
+THIRD_PARTY_INCLUDES_START
+#include <XUser.h>
+THIRD_PARTY_INCLUDES_END
+#include "Microsoft/HideMicrosoftPlatformTypes.h"
+
+namespace GDKHandleTraits
+{
+	template<typename T>
+	static int32 CompareHandle(T Lhs, T Rhs)
+	{
+		return (Lhs == Rhs) ? 0 : 1;
+	}
+
+	template<typename T>
+	static void DuplicateHandle(T Source, T* Target)
+	{
+	}
+
+	template<typename T>
+	static void ReleaseHandle(T Handle)
+	{
+	}
+}
+
+template <typename HandleType>
+class TGDKHandle
+{
+public:
+	using ThisType = TGDKHandle<HandleType>;
+
+	TGDKHandle() :
+		Handle{ nullptr }
+	{
+	}
+	~TGDKHandle()
+	{
+		Clear();
+	}
+
+	explicit TGDKHandle(HandleType source) :
+		Handle{ source }
+	{
+	}
+
+	TGDKHandle(const ThisType& source) :
+		Handle{ nullptr }
+	{
+		if (source.Handle)
+		{
+			GDKHandleTraits::DuplicateHandle(source.Handle, &Handle);
+		}
+	}
+
+	ThisType &operator=(const ThisType& source)
+	{
+		Clear();
+
+		if (source.Handle)
+		{
+			GDKHandleTraits::DuplicateHandle(source.Handle, &Handle);
+		}
+		return *this;
+	}
+
+	TGDKHandle(ThisType&& moveFrom) :
+		Handle{ moveFrom.Handle }
+	{
+		moveFrom.Handle = nullptr;
+	}
+
+	ThisType& operator=(ThisType&& moveFrom)
+	{
+		Clear();
+
+		Handle = moveFrom.Handle;
+		moveFrom.Handle = nullptr;
+
+		return *this;
+	}
+
+	bool operator==(const ThisType& Rhs) const
+	{
+		return (GDKHandleTraits::CompareHandle<HandleType>(*this, Rhs) == 0);
+	}
+
+	operator HandleType() const
+	{
+		return Handle;
+	}
+
+	operator const HandleType*() const
+	{
+		return IsValid() ? &Handle : nullptr;
+	}
+
+	operator bool() const
+	{
+		return IsValid();
+	}
+
+	bool IsValid() const
+	{
+		return Handle != nullptr;
+	}
+
+	void Reset(HandleType handle)
+	{
+		Clear();
+
+		Handle = handle;
+	}
+
+	void Clear()
+	{
+		if (Handle)
+		{
+			GDKHandleTraits::ReleaseHandle(Handle);
+			Handle = nullptr;
+		}
+	}
+
+	HandleType* GetInitReference()
+	{
+		checkf(Handle == nullptr,TEXT("GDK handle already initalised"));
+		return &Handle;
+	}
+
+	HandleType* ReleaseAndGetInitReference()
+	{
+		Clear();
+		return &Handle;
+	}
+
+	/** Needed for TMap::GetTypeHash() */
+	friend SIZE_T GetTypeHash(const ThisType& A)
+	{
+		return reinterpret_cast<SIZE_T>(A.Handle);
+	}
+
+private:
+	HandleType Handle;
+};
+
+#define SETHANDLETYPETRAITS(handleType, duplicateMethod, releaseMethod) \
+	namespace GDKHandleTraits { \
+		template<> void DuplicateHandle(handleType Source, handleType* Target) \
+		{ duplicateMethod(Source, Target); } \
+		template<> void ReleaseHandle(handleType Handle) \
+		{ releaseMethod(Handle); } \
+	}
+
+#define SETHANDLETYPETRAITSEX(handleType, duplicateMethod, releaseMethod, compareMethod) \
+	namespace GDKHandleTraits { \
+		template<> void DuplicateHandle(handleType Source, handleType* Target) \
+		{ duplicateMethod(Source, Target); } \
+		template<> void ReleaseHandle(handleType Handle) \
+		{ releaseMethod(Handle); } \
+		template<> int32 CompareHandle(handleType Lhs, handleType Rhs) \
+		{ return compareMethod(Lhs, Rhs); } \
+	}
+
+SETHANDLETYPETRAITS(XTaskQueueHandle, XTaskQueueDuplicateHandle, XTaskQueueCloseHandle);
+SETHANDLETYPETRAITSEX(XUserHandle, XUserDuplicateHandle, XUserCloseHandle, XUserCompare);
+
+typedef TGDKHandle<XTaskQueueHandle> FGDKTaskQueueHandle;
+typedef TGDKHandle<XUserHandle> FGDKUserHandle;
+
+
+#endif //WITH_GRDK
diff --git a/Engine/Source/Runtime/Microsoft/GDK/GDKRuntime/Public/GDKHandleTracker.h b/Engine/Source/Runtime/Microsoft/GDK/GDKRuntime/Public/GDKHandleTracker.h
new file mode 100644
--- /dev/null
+++ b/Engine/Source/Runtime/Microsoft/GDK/GDKRuntime/Public/GDKHandleTracker.h
@@ -0,0 +1,46 @@
+// Copyright Epic Games, Inc. All Rights Reserved.
+
+#pragma once
+
+#include "CoreTypes.h"
+
+#define UE_API GDKRUNTIME_API
+
+/* Whether to enable XSystemHandleTrack support. 
+ *
+ * It can be activated in a number of ways:
+ *     Use the   -gdkhandletrack   command-line parameter for global tracking (unclosed handles will be dumped on exit)
+ *     Use GDK.XSystemHandle.TrackStart / GDK.XSystemHandle.TrackStop for on-demand tracking.
+ *     Use GDK_SCOPE_TRACK_HANDLES() to track any handles created in the current scope (expected use case is for localized debugging only)
+ * 
+ * Note that only one can be active at any one time
+ */
+#ifndef WITH_GDK_HANDLE_TRACKER
+	#define WITH_GDK_HANDLE_TRACKER (!UE_BUILD_SHIPPING && !ENABLE_PGO_PROFILE && WITH_GRDK)
+#endif
+
+#if WITH_GDK_HANDLE_TRACKER
+
+	class FGDKHandleTracker
+	{
+	public:
+		inline FGDKHandleTracker()   { Start(); }
+		inline ~FGDKHandleTracker()  { Stop(); }
+
+		static UE_API void Start();
+		static UE_API void Stop();
+	};
+
+	// track GDK handles that were created while the given scope was active. (note: this tracks all GDK handles in all threads.)
+	// this is expected to be used for very localized debugging
+	#define GDK_SCOPE_TRACK_HANDLES() FGDKHandleTracker ANONYMOUS_VARIABLE( GDKHandleTrack )
+
+#else
+
+	// no-op
+	#define GDK_SCOPE_TRACK_HANDLES()
+
+
+#endif //WITH_GDK_HANDLE_TRACKER
+
+#undef UE_API
diff --git a/Engine/Source/Runtime/Microsoft/GDK/GDKRuntime/Public/GDKRuntimeModule.h b/Engine/Source/Runtime/Microsoft/GDK/GDKRuntime/Public/GDKRuntimeModule.h
new file mode 100644
--- /dev/null
+++ b/Engine/Source/Runtime/Microsoft/GDK/GDKRuntime/Public/GDKRuntimeModule.h
@@ -0,0 +1,286 @@
+// Copyright Epic Games, Inc. All Rights Reserved.
+
+#pragma once
+
+#if !defined(WITH_GRDK)
+	#define WITH_GRDK 1
+#endif
+
+#if WITH_GRDK
+
+#include "CoreMinimal.h"
+#include "Modules/ModuleInterface.h"
+#include "Modules/ModuleManager.h"
+#include "Templates/Function.h"
+#include "Delegates/Delegate.h"
+#include "Misc/Optional.h"
+#include "GenericPlatform/GenericPlatformMisc.h"
+#include "GDKHandle.h"
+
+#include "Microsoft/AllowMicrosoftPlatformTypes.h"
+THIRD_PARTY_INCLUDES_START
+#include <XAsync.h>
+THIRD_PARTY_INCLUDES_END
+#include "Microsoft/HideMicrosoftPlatformTypes.h"
+
+
+enum class XUserGamertagComponent: uint32_t;
+enum class XNetworkingConnectivityLevelHint : uint32_t;
+
+GDKRUNTIME_API DECLARE_LOG_CATEGORY_EXTERN(LogGDK, Log, All);
+
+using FGDKRuntimeErrorProcessDelegate = TTSDelegate<bool(const char* /*Msg*/)>;
+using FGDKPickUserCompleteDelegate = TTSDelegate<void(HRESULT, FGDKUserHandle)>;
+
+DECLARE_MULTICAST_DELEGATE_OneParam( FGDKOnGameInviteReceived, const FString&/*Uri*/ );
+typedef FGDKOnGameInviteReceived::FDelegate FGDKOnGameInviteReceivedDelegate;
+
+struct XblUserProfile;
+
+/**
+ * Setup GDK runtime if there is any code rely on it before loading IGDKRuntimeModule
+ */
+GDKRUNTIME_API bool SetupGDKEnvironment( const TCHAR* ConfigIniSection = nullptr );
+
+/**
+ * Shutdown GDK runtime, symmetrical call to SetupGDKEnvironment
+ */
+GDKRUNTIME_API void TeardownGDKEnvironment();
+
+class IGDKRuntimeModule : public IModuleInterface
+{
+public:
+	/**
+	 * Singleton-like access to the GDK runtime module instance
+	 * @return Returns IGDKRuntimeModule singleton instance, loading the module on demand if needed
+	 */
+	static inline IGDKRuntimeModule& Get()
+	{
+		return FModuleManager::LoadModuleChecked<IGDKRuntimeModule>("GDKRuntime");
+	}
+
+	/**
+	 * Singleton-like access to the GDK runtime module instance
+	 * @return Returns IGDKRuntimeModule singleton instance if the module is already loaded
+	 */
+	static inline IGDKRuntimeModule* TryGet()
+	{
+		return FModuleManager::GetModulePtr<IGDKRuntimeModule>("GDKRuntime");
+	}
+
+	/**
+	 * Returns whether the GDK runtime was initialized successfully.
+	 */
+	virtual bool IsAvailable() const = 0;
+
+
+	/**
+	 * Returns the Title Id for the game. See the GDK documentation for details.
+	 */
+	virtual uint32 GetTitleId() const = 0;
+
+	/**
+	 * Returns the current XBL sandbox for the platform. See the GDK documentation for details.
+	 */
+	virtual FString GetXboxSandboxId() const = 0;
+
+	/**
+	 * Returns the SCID for the title. See the GDK documentation for details.
+	 * If bAllowPlaceholder is true and no SCID has been defined, a placeholder one will be returned rather than an empty string.
+	 */
+	virtual FString GetPrimaryServiceConfigId( bool bAllowPlaceholder = true ) const = 0;
+
+	/**
+	 * Returns the gamertag for the given user handle
+	 * @param UserHandle the user to query
+	 * @return the gamertag for the user, or an empty string if the user is not valid
+	 */
+	virtual FString GetGamertag(const FGDKUserHandle& UserHandle) const = 0;
+
+	/**
+	 * Returns the gamertag for the given user handle using the given gamertag convension
+	 * @param UserHandle        the user to query
+	 * @param GamertagComponent the convention to use - Classic/Modern etc.
+	 * @return the gamertag for the user, or an empty string if the user is not valid
+	 */
+	virtual FString GetGamertag(const FGDKUserHandle& UserHandle, XUserGamertagComponent GamertagComponent) const = 0;
+
+	/**
+	 * Returns the preferred default gametag component
+	 */
+	virtual XUserGamertagComponent GetDefaultGamertagComponent() const = 0;
+
+	/**
+	 * Returns the config section that the GDK configuration is stored in
+	 * @return [Config Section] name
+	 */
+	virtual const TCHAR* GetConfigSectionName() const = 0;
+
+	/**
+	 * Returns whether the project was configured to use the simplified user model
+	 */
+	virtual bool IsUsingSimplifiedUserModel() const = 0;
+
+	/**
+	 * Returns whether networking has been initialized yet (see XNetworkingConnectivityHint in the GDK documentation for details.)
+	 */
+	virtual bool IsNetworkInitialized() const = 0;
+
+
+	/** Simple wrapper for async GDK task execution. 
+	 * @param InitFunction   code that will be executed on the game thread to start the operation (e.g. XUserAddAsync). Expected to return the result from the GDK async function.
+	 * @param ResultFunction optional code that will be executed when the async operation has finished (e.g. XUserAddResult). Execution thread depends on TaskQueue.
+	 * @param TaskQueue      optional task queue. Default is GetGenericTaskQueue which executes on a background thread and ResultFunction happens on the game thread.
+	 * @return the result from InitFunction, so S_OK means the async task has started successfully and ResultFunction will be called
+	*/
+	virtual HRESULT AsyncGDKTask( TFunction<HRESULT(XAsyncBlock*)> InitFunction, TFunction<void(XAsyncBlock*)> ResultFunction = nullptr, XTaskQueueHandle TaskQueue = nullptr ) const = 0;
+
+	/** Simple wrapper for cancellable async GDK task execution. 
+	 * @param OutMonitor     variable that will receive the task monitor, allowing cancellation of the task (ResultFunction will still be called, with E_ABORT)
+	 * @param InitFunction   code that will be executed on the game thread to start the operation (e.g. XUserAddAsync). Expected to return the result from the GDK async function.
+	 * @param ResultFunction optional code that will be executed when the async operation has finished (e.g. XUserAddResult). Execution thread depends on TaskQueue.
+	 * @param TaskQueue      optional task queue. Default is GetGenericTaskQueue which executes on a background thread and ResultFunction happens on the game thread.
+	 * @return the result from InitFunction, so S_OK means the async task has started successfully and ResultFunction will be called
+	*/
+	virtual HRESULT AsyncGDKTask( class FGDKAsyncTaskMonitor& OutMonitor, TFunction<HRESULT(XAsyncBlock*)> InitFunction, TFunction<void(XAsyncBlock*)> ResultFunction = nullptr, XTaskQueueHandle TaskQueue = nullptr ) const = 0;
+
+
+	/**
+	 * Get a reusable task queue. Tasks execute on a background thread and completion callbacks happen on the game thread
+	 */
+	virtual XTaskQueueHandle GetGenericTaskQueue() const = 0;
+
+	/**
+	 * Get a reusable task queue. Tasks execute on a background thread and completion callbacks also happen on a background thread
+	 */
+	virtual XTaskQueueHandle GetBackgroundTaskQueue() const = 0;
+
+
+	/**
+	 * Returns the user handle for the given platform user id
+	 * @param PlatformUserId the user to query
+	 * @return user handle for the given platform user
+	 */
+	virtual FGDKUserHandle GetUserHandleByPlatformId(FPlatformUserId PlatformUserId) const = 0;
+
+	/**
+	 * Returns the user handle for the given XUID
+	 * @param XUserId XUID to query
+	 * @return user handle for the given Xbox User Id
+	 */
+	virtual FGDKUserHandle GetUserHandleByXUserId(uint64 XUserId) const = 0;
+
+	/**
+	 * Returns the platform user id for the given user handle
+	 * @param UserHandle the user to query
+	 * @return platform user id for the given user
+	 */
+	virtual FPlatformUserId GetPlatformIdByUserHandle(const FGDKUserHandle& UserHandle) = 0;
+
+	/**
+	 * Returns the XUID for the given platform user
+	 * @param PlatformUserId the user to query
+	 * @return Xbox User Id for the given platform user
+	 */
+	virtual TOptional<uint64> GetXUserIdByPlatformId(FPlatformUserId PlatformUserId) const = 0;
+
+	/**
+	 * Returns all current user handles
+	 */
+	virtual TArray<FGDKUserHandle> GetAllUserHandles() const = 0;
+
+	/**
+	 * Returns whether the given local user was previously signed in but is not presently signed in
+	 * @param UserLocalId the user to query
+	 * @return true if the user had previously signed in but has now signed out. Will return false when the user has completed signing back in
+	 * @note this function uses XUserLocalId because the user handle would be invalid once the user is signed out
+	 */
+	virtual bool WasUserRecentlySignedOut(XUserLocalId UserLocalId) const = 0;
+
+	/**
+	 * The async call to pick a user through the system UI
+	 * @param Options The options when show the user pick UI
+	 * @param Delegate The delegate to receive the HRESULT of the call
+	 */
+	virtual void PickUserAsync(XUserAddOptions Options, FGDKPickUserCompleteDelegate Delegate) const = 0;
+
+	/**
+	 * Returns the platform user id that will be used to represent any gamepads that are not paired to a user
+	 */
+	virtual FPlatformUserId GetUnpairedPlatformId() const = 0;
+
+	/**
+	 * Set runtime error process delegate for specific gdk runtime error
+	 */
+	virtual void SetRuntimeErrorProcessDelegate(HRESULT hResult, FGDKRuntimeErrorProcessDelegate Delegate) = 0;
+
+	/**
+	 * Clear runtime error process delegate for specific gdk runtime error
+	 */
+	virtual void ClearRuntimeErrorProcessDelegate(HRESULT hResult) = 0;
+
+	/**
+	 * Register a callback for when the game receives an online game invitation. callback parameter is the matchmaking Uri
+	 */
+	virtual FDelegateHandle RegisterGameInviteReceivedDelegate( FGDKOnGameInviteReceivedDelegate Delegate ) = 0;
+
+	/**
+	 * Unregister a previously-registered game invite callback
+	 */
+	virtual void UnregisterGameInviteReceivedDelegate( const FDelegateHandle& DelegateHandle ) = 0;
+
+
+	// ... FPlatformMisc support functions ...
+
+	/**
+	 * Show a message box if possible, otherwise print a message and return the default
+	 * @param MsgType What sort of options are provided
+	 * @param Text Specific message
+	 * @param Caption String indicating the title of the message box
+	 * @return Very strange convention...not really EAppReturnType, see implementation
+	 */
+	virtual EAppReturnType::Type MessageBoxExt(EAppMsgType::Type MsgType, const TCHAR* Text, const TCHAR* Caption) const = 0;
+
+	/**
+	 * Returns the current network connection type
+	 */
+	virtual ENetworkConnectionType GetNetworkConnectionType() const = 0;
+
+	/**
+	 * Returns whether WiFi connection is currently active
+	 */
+	virtual bool HasActiveWiFiConnection() const = 0;
+
+	/**
+	 * Prevents the screen from dimming, locking or powering down
+	 */
+	virtual void DisableScreenTimeout() = 0;
+
+	/**
+	 * Allows the screen to dimming, lock or powering down again
+	 */
+	virtual void EnableScreenTimeout() = 0;
+
+
+
+	// ... internal functions ...
+
+	virtual void Internal_HandlePlatformTextFieldVisible( bool bIsVisible ) = 0;
+	virtual bool Internal_IsPlatformTextFieldVisible() const = 0;
+
+	// ... PIE support ...
+#if WITH_EDITOR
+	virtual void Internal_InitForPIE() = 0;
+	virtual void Internal_TeardownForPIE() = 0;
+	virtual TMulticastDelegate<void()>& GetOnInitForPIE() = 0;
+	virtual TMulticastDelegate<void()>& GetOnTeardownForPIE() = 0;
+#endif
+};
+
+
+
+GDKRUNTIME_API const TCHAR* LexToString(XNetworkingConnectivityLevelHint Value);
+GDKRUNTIME_API FString LexToString(const FGDKUserHandle& UserHandle);
+GDKRUNTIME_API FString LexToString(const XUserLocalId& UserId);
+
+#endif //WITH_GRDK
diff --git a/Engine/Source/Runtime/Microsoft/GDK/GDKRuntime/Public/GDKTaskQueueHelpers.h b/Engine/Source/Runtime/Microsoft/GDK/GDKRuntime/Public/GDKTaskQueueHelpers.h
new file mode 100644
--- /dev/null
+++ b/Engine/Source/Runtime/Microsoft/GDK/GDKRuntime/Public/GDKTaskQueueHelpers.h
@@ -0,0 +1,232 @@
+// Copyright Epic Games, Inc. All Rights Reserved.
+
+#pragma once
+
+#include "GDKRuntimeModule.h"
+#if WITH_GRDK
+#include "CoreTypes.h"
+#include "Templates/Function.h"
+#include "Templates/SharedPointer.h"
+#include "Tasks/Task.h"
+
+#include "Microsoft/AllowMicrosoftPlatformTypes.h"
+THIRD_PARTY_INCLUDES_START
+#include <XTaskQueue.h>
+#include <XAsync.h>
+THIRD_PARTY_INCLUDES_END
+#include "Microsoft/HideMicrosoftPlatformTypes.h"
+
+#define UE_API GDKRUNTIME_API
+
+/*	Wrapper for XAsyncBlock that allows capture variables in the lambda 
+
+	If lifetime management is not required, using AsyncGDKTask() is recommended instead
+*/
+typedef TFunction<void(class FGDKAsyncBlock*)> FGDKAsyncBlockDelegate;
+
+class FGDKAsyncBlock
+{
+public:
+	UE_API FGDKAsyncBlock(void* UserData, FGDKAsyncBlockDelegate Delegate, XTaskQueueHandle TaskQueue = nullptr);
+
+	UE_API ~FGDKAsyncBlock();
+
+	UE_API HRESULT GetStatus();
+
+	UE_API XAsyncBlock* GetInnerBlockForGDKAPI();
+	UE_API void* GetUserData();
+	operator XAsyncBlock*() { return GetInnerBlockForGDKAPI(); }
+
+protected:
+	XAsyncBlock AsyncBlock;
+	void* UserData;
+	FGDKAsyncBlockDelegate Delegate;
+
+	UE_API void DelegateWrapper();
+};
+
+
+
+/* 
+	Wrapper for XTaskQueueHandle that manages clean shutdown automatically
+
+	This can be used when a separate, custom task queue is required. In most cases FGDKAsyncTaskQueue::GetGenericQueue() should be sufficient
+*/
+class FGDKAsyncTaskQueue
+{
+public:
+	UE_API FGDKAsyncTaskQueue( XTaskQueueDispatchMode WorkDispatchMode = XTaskQueueDispatchMode::ThreadPool, XTaskQueueDispatchMode CompletionDispatchMode = XTaskQueueDispatchMode::ThreadPool );
+	UE_API ~FGDKAsyncTaskQueue();
+
+	FORCEINLINE XTaskQueueHandle GetQueue() const { return Queue; }
+	UE_API HRESULT BlockUntilComplete( XAsyncBlock& AsyncBlock ) const;
+	UE_API void CancelPendingTasksAndDestroyQueue();
+
+
+	// access to common task queues
+	static UE_API XTaskQueueHandle GetGenericQueue();			// completion callbacks happen on the game thread
+	static UE_API XTaskQueueHandle GetBackgroundTaskQueue();	// completion callbacks happen on a background thread
+
+	// internal use - startup & shutdown the queue system
+	static UE_API void PlatformInit();
+	static UE_API void PlatformTeardown();
+
+protected:
+	XTaskQueueDispatchMode WorkDispatchMode;
+	XTaskQueueDispatchMode CompletionDispatchMode;
+	XTaskQueueHandle Queue;
+};
+
+
+
+/* Helper class for executing simple GDK async tasks locally. This is expected to be created on the stack. 
+   (Use LocalGDKTask() for an easier to use equivalent.)
+
+	FGDKLocalTaskBlock Block;
+	if (SUCCEEDED( XUserAddAsync( XUserAddOptions::None, Block ) ) )
+	{
+		AsyncBlock.BlockUntilComplete();
+
+		FGDKUserHandle User;
+		XUserAddResult(Block, User.GetInitReference());
+	}
+*/
+class FGDKLocalTaskBlock
+{
+public:
+	UE_API FGDKLocalTaskBlock();
+	UE_API ~FGDKLocalTaskBlock();
+
+	operator XAsyncBlock*() { return &AsyncBlock; }
+
+	UE_API HRESULT BlockUntilComplete();
+	UE_API HRESULT GetStatus();
+
+protected:
+	XAsyncBlock AsyncBlock;
+	FGDKAsyncTaskQueue Queue;
+};
+
+
+
+/* Simple wrapper for async GDK task execution. Returns the result from InitFunction, so S_OK means the async task has started successfully
+
+    AsyncGDKTask( 
+		[]( XAsyncBlock* Block )
+		{
+			return XUserAddAsync( XUserAddOptions::None, Block );
+		},
+		[]( XAsyncBlock* Block )
+		{
+			FGDKUserHandle User;
+			XUserAddResult(Block, User.GetInitReference());
+		}
+	);
+     
+*/
+UE_API HRESULT AsyncGDKTask( TFunction<HRESULT(XAsyncBlock*)> InitFunction, TFunction<void(XAsyncBlock*)> ResultFunction = nullptr, XTaskQueueHandle TaskQueue = nullptr );
+
+
+
+
+/* Simple wrapper for async GDK task execution that can be monitored via a helper class. Returns the result from InitFunction, so S_OK means the async task has started successfully
+
+	FGDKAsyncTaskMonitor MyTask;
+    AsyncGDKTask( 
+		MyTask,
+		[]( XAsyncBlock* Block )
+		{
+			return XUserAddAsync( XUserAddOptions::None, Block );
+		},
+		[]( XAsyncBlock* Block )
+		{
+			FGDKUserHandle User;
+			XUserAddResult(Block, User.GetInitReference());
+		}
+	);
+	MyTask.TryCancel(true);
+
+*/
+UE_API HRESULT AsyncGDKTask( class FGDKAsyncTaskMonitor& OutAsyncTaskMonitor, TFunction<HRESULT(XAsyncBlock*)> InitFunction, TFunction<void(XAsyncBlock*)> ResultFunction = nullptr, XTaskQueueHandle TaskQueue = nullptr );
+
+
+
+
+/* Simple wrapper for local GDK task execution. Blocks until the task is completed. Returns the result from InitFunction or ResultFunction
+
+    LocalGDKTask( 
+		[]( XAsyncBlock* Block )
+		{
+			return XUserAddAsync( XUserAddOptions::None, Block );
+		},
+		[]( XAsyncBlock* Block )
+		{
+			FGDKUserHandle User;
+			return XUserAddResult(Block, User.GetInitReference());
+		}
+	);
+     
+*/
+UE_API HRESULT LocalGDKTask( TFunction<HRESULT(XAsyncBlock*)> InitFunction, TFunction<HRESULT(XAsyncBlock*)> ResultFunction = nullptr );
+
+
+/* Simple UE::Task wrapper for GDK task execution. Task result contains the result from InitFunction or ResultFunction
+
+ FGDKTask Task = LaunchGDKTask( 
+		UE_SOURCE_LOCATION,
+		[]( XAsyncBlock* Block )
+		{
+			return XUserAddAsync( XUserAddOptions::None, Block );
+		},
+		[]( XAsyncBlock* Block )
+		{
+			FGDKUserHandle User;
+			return XUserAddResult(Block, User.GetInitReference());
+		}
+	);
+*/
+typedef UE::Tasks::TTask<HRESULT> FGDKTask;
+
+// launches a gdk task for asynchronous execution
+// returns a trivially relocatable instance that can be used to wait for task completion or to obtain task execution result
+UE_API FGDKTask LaunchGDKTask( 
+	const TCHAR* DebugName, 
+	TUniqueFunction<HRESULT(XAsyncBlock*)> InitFunction, 
+	TUniqueFunction<HRESULT(XAsyncBlock*)> ResultFunction, 
+	UE::Tasks::ETaskPriority Priority = UE::Tasks::ETaskPriority::Normal );
+
+UE_API FGDKTask LaunchGDKTask( 
+	const TCHAR* DebugName, 
+	class FGDKAsyncTaskMonitor& OutAsyncTaskMonitor,
+	TUniqueFunction<HRESULT(XAsyncBlock*)> InitFunction, 
+	TUniqueFunction<HRESULT(XAsyncBlock*)> ResultFunction, 
+	UE::Tasks::ETaskPriority Priority = UE::Tasks::ETaskPriority::Normal );
+
+
+
+
+
+
+
+/* Helper allowing cancellation from AsyncGDKTask() and LaunchGDKTask() */
+class FGDKAsyncTaskMonitor
+{
+public:
+	// Attempt to cancel the task. Note that the ResultFunction callback will be called and the GDK API function will likely return E_ABORT
+	//  - When bWait is false, returns true if cancel request was made. 
+	//  - When bWait is true it returns true if the task was successfully cancelled. This will happen before the ResultFunction is called
+	//  - Returns false in all other failure cases, including if the task has already completed.
+	UE_API bool TryCancel(bool bWait = true);
+
+	// Determines if the task that we're monitoring is still valid. This will be true until the ResultFunction callback has completed
+	UE_API bool IsValid() const;
+
+protected:
+	friend void Internal_InitAsyncTaskMonitor( FGDKAsyncTaskMonitor& TaskMonitor, TSharedPtr<FGDKAsyncBlock> Block );
+	TWeakPtr<FGDKAsyncBlock> WeakBlock;
+};
+
+
+#undef UE_API
+
+#endif //WITH_GRDK
diff --git a/Engine/Source/Runtime/Microsoft/GDK/GDKRuntime/Public/GDKThreadCheck.h b/Engine/Source/Runtime/Microsoft/GDK/GDKRuntime/Public/GDKThreadCheck.h
new file mode 100644
--- /dev/null
+++ b/Engine/Source/Runtime/Microsoft/GDK/GDKRuntime/Public/GDKThreadCheck.h
@@ -0,0 +1,71 @@
+// Copyright Epic Games, Inc. All Rights Reserved.
+
+#pragma once
+
+#include "CoreTypes.h"
+
+// whether to enable XThreadSetTimeSensitive for RHI, Render and Game threads (depending on WITH_GDK_THREAD_CHECK_ALWAYS or -TimeSensitiveThreadCheck too)
+#ifndef WITH_GDK_THREAD_CHECK
+	#define WITH_GDK_THREAD_CHECK (!UE_BUILD_SHIPPING && !UE_BUILD_TEST && !ENABLE_PGO_PROFILE && WITH_GRDK)
+#endif
+
+// when set to 1 the -TimeSensitiveThreadCheck command line parameter is not needed
+#ifndef WITH_GDK_THREAD_CHECK_ALWAYS
+	#define WITH_GDK_THREAD_CHECK_ALWAYS 0
+#endif
+
+
+#if WITH_GDK_THREAD_CHECK
+
+	class FGDKScopedNotTimeSensitive
+	{
+	public:
+#if IS_MONOLITHIC
+		FGDKScopedNotTimeSensitive();
+		GDKRUNTIME_API ~FGDKScopedNotTimeSensitive();
+#else
+		inline FGDKScopedNotTimeSensitive();
+		inline ~FGDKScopedNotTimeSensitive();
+#endif //IS_MONOLITHIC
+
+	private:
+		bool bWasTimeSensitive;
+	};
+
+	// temporarily disables the time sensitive thread check in the current scope
+	#define GDK_SCOPE_NOT_TIME_SENSITIVE() FGDKScopedNotTimeSensitive ANONYMOUS_VARIABLE( GDKNotTimeSenstitive )
+
+
+	// runtime check if time sensitive thread checking should be enabled
+	extern bool GDKRUNTIME_API IsGDKTimeSensitiveThreadCheckEnabled();
+
+#else
+
+	// no-op
+	#define GDK_SCOPE_NOT_TIME_SENSITIVE()
+
+#endif //WITH_GDK_THREAD_CHECK
+
+
+
+
+// NOTE: modular game builds are not a supported feature
+// this is here just to make sure it can still compile
+#if WITH_GDK_THREAD_CHECK && !IS_MONOLITHIC
+
+#include "Microsoft/AllowMicrosoftPlatformTypes.h"
+#include <XThread.h>
+#include "Microsoft/HideMicrosoftPlatformTypes.h"
+
+inline FGDKScopedNotTimeSensitive::FGDKScopedNotTimeSensitive()
+{
+	bWasTimeSensitive = XThreadIsTimeSensitive();
+	XThreadSetTimeSensitive(false);
+}
+
+inline FGDKScopedNotTimeSensitive::~FGDKScopedNotTimeSensitive()
+{
+	XThreadSetTimeSensitive(bWasTimeSensitive);
+}
+#endif //IS_MONOLITHIC
+
diff --git a/Engine/Source/Runtime/Microsoft/GDK/GDKSaveGameSystem/GDKSaveGameSystem.Build.cs b/Engine/Source/Runtime/Microsoft/GDK/GDKSaveGameSystem/GDKSaveGameSystem.Build.cs
new file mode 100644
--- /dev/null
+++ b/Engine/Source/Runtime/Microsoft/GDK/GDKSaveGameSystem/GDKSaveGameSystem.Build.cs
@@ -0,0 +1,23 @@
+// Copyright Epic Games, Inc. All Rights Reserved.
+
+using UnrealBuildTool;
+
+public class GDKSaveGameSystem : ModuleRules
+{
+	public GDKSaveGameSystem( ReadOnlyTargetRules Target) : base(Target)
+	{
+		if (Target.Platform == UnrealTargetPlatform.Win64)
+		{
+			BinariesSubFolder = "GDK";
+		}
+
+		PrivateDependencyModuleNames.AddRange(new string[]
+			{
+				"Core",
+				"Engine",
+				"ApplicationCore",
+				"GRDK",
+				"GDKRuntime",
+			});
+	}
+}
diff --git a/Engine/Source/Runtime/Microsoft/GDK/GDKSaveGameSystem/Private/GDKSaveGameSystem.cpp b/Engine/Source/Runtime/Microsoft/GDK/GDKSaveGameSystem/Private/GDKSaveGameSystem.cpp
new file mode 100644
--- /dev/null
+++ b/Engine/Source/Runtime/Microsoft/GDK/GDKSaveGameSystem/Private/GDKSaveGameSystem.cpp
@@ -0,0 +1,823 @@
+// Copyright Epic Games, Inc. All Rights Reserved.
+
+#include "GDKSaveGameSystem.h"
+#if WITH_GRDK
+#include "GDKRuntimeModule.h"
+#include "GDKTaskQueueHelpers.h"
+#include "GDKThreadCheck.h"
+#include "GameDelegates.h"
+#include "Async/Async.h"
+#include "Misc/ConfigCacheIni.h"
+#include "Misc/DateTime.h"
+#include "HAL/IConsoleManager.h"
+#include "HAL/PlatformMisc.h"
+#include "HAL/PlatformApplicationMisc.h"
+#include "HAL/ThreadHeartBeat.h"
+
+#include "Microsoft/AllowMicrosoftPlatformTypes.h"
+#include <XThread.h>
+#include "Microsoft/HideMicrosoftPlatformTypes.h"
+
+
+DEFINE_LOG_CATEGORY(LogGDKSaveGame);
+
+static int32 GGDKUseXDKCompatibleSave = 0;
+static FAutoConsoleVariableRef CVarGDKUseXDKCompatibleSave(
+	TEXT("GDK.UseXDKCompatibleSave"),
+	GGDKUseXDKCompatibleSave,
+	TEXT("Whether the savegame system will read and write XDK save games."),
+	ECVF_Default
+);
+
+static int32 GGDKSaveGameSyncOnDemand = 0;
+static FAutoConsoleVariableRef CVarGDKSaveGameSyncOnDemand(
+	TEXT("GDK.SaveGame.UseSyncOnDemand"),
+	GGDKSaveGameSyncOnDemand,
+	TEXT("Sync saves on demand. Only necessary for games with a lot of very large containers. See the GDK documentation for details."),
+	ECVF_Default
+);
+
+
+
+const uint32 FGDKSaveGameSystem::MaxContainerName = 256;
+const char* FGDKSaveGameSystem::DefaultBlobName = "Data";
+
+const char* FGDKSaveGameSystem::XDKCompatibleContainerName = "UE4DefaultSaveGameContainer";
+
+
+FGDKSaveGameSystem::FGDKSaveGameSystem()
+{
+	GDK_SCOPE_NOT_TIME_SENSITIVE(); // (startup only) XUserRegisterForChangeEvent is not safe to call on a time-sensitive thread
+
+	//register for a callback when a user signs out
+	auto UserChangeFn = [](void* Context, XUserLocalId UserLocalId, XUserChangeEvent Event)
+	{
+		((FGDKSaveGameSystem*)Context)->OnUserChange(UserLocalId, Event);
+	};
+	XUserRegisterForChangeEvent(FGDKAsyncTaskQueue::GetGenericQueue(), this, UserChangeFn, &UserChangeCallbackToken);
+
+}
+
+FGDKSaveGameSystem::~FGDKSaveGameSystem()
+{
+	GDK_SCOPE_NOT_TIME_SENSITIVE(); // (shutdown only) XUserUnregisterForChangeEvent & XGameSaveCloseProvider are not safe to call on a time-sensitive thread
+
+	XUserUnregisterForChangeEvent(UserChangeCallbackToken, true);
+
+	for (auto Itr : PerUserSaveProviders)
+	{
+		XGameSaveCloseProvider(Itr.Value);
+	}
+	PerUserSaveProviders.Reset();
+}
+
+void FGDKSaveGameSystem::OnUserChange(XUserLocalId UserLocalId, XUserChangeEvent Event)
+{
+	// if the user signs out, remove them from the map - otherwise we risk E_GS_HANDLE_EXPIRED
+	if (Event == XUserChangeEvent::SignedOut)
+	{
+		GDK_SCOPE_NOT_TIME_SENSITIVE(); // XUserFindUserByLocalId and XGameSaveCloseProvider are not safe to call on a time-sensitive thread
+	
+		FGDKUserHandle UserHandle;
+		XUserFindUserByLocalId(UserLocalId, UserHandle.GetInitReference());
+
+		XGameSaveProviderHandle Provider = nullptr;
+		if (PerUserSaveProviders.RemoveAndCopyValue(UserHandle, Provider))
+		{
+			XGameSaveCloseProvider(Provider);
+		}
+	}
+}
+
+
+UE::Tasks::FTask FGDKSaveGameSystem::InitSaveGameProviderAsync(FPlatformUserId PlatformUserId)
+{
+	return AsyncTaskPipe.Launch(UE_SOURCE_LOCATION,
+		[this, PlatformUserId]
+		{
+			const FGDKUserHandle User = IGDKRuntimeModule::Get().GetUserHandleByPlatformId(PlatformUserId);
+			if (!User.IsValid())
+			{
+				UE_LOG(LogGDKSaveGame, Warning, TEXT("User %d is not valid"), PlatformUserId.GetInternalId());
+				return;
+			}
+
+			XGameSaveProviderHandle Provider = PerUserSaveProviders.FindRef(User);
+			if (Provider == nullptr)
+			{
+				FString PrimaryServiceConfigId = IGDKRuntimeModule::Get().GetPrimaryServiceConfigId(false);
+				if (PrimaryServiceConfigId.IsEmpty())
+				{
+					UE_LOG(LogGDKSaveGame, Warning, TEXT("PrimaryServiceConfigId is not set - please make sure Project Settings are configured correctly"));
+					return;
+				}
+
+				// add a nested task to initialize the provider. this may display on-screen UI items
+				// our task will not complete until the nested task has completed.
+				AddNested( LaunchGDKTask( UE_SOURCE_LOCATION,
+					[this, User, PrimaryServiceConfigId](XAsyncBlock* Block)
+					{
+						bool bSyncOnDemand = (GGDKSaveGameSyncOnDemand != 0);
+
+						UE_LOG(LogGDKSaveGame, Log, TEXT("Initializing save provider for user %s..."), *LexToString(User));
+						HRESULT hResult = XGameSaveInitializeProviderAsync(User, TCHAR_TO_UTF8(*PrimaryServiceConfigId), bSyncOnDemand, Block);
+						ProcessErrorCode(hResult, TEXT("InitializeProviderInit"));
+
+						return hResult;
+					},
+					[this, User](XAsyncBlock* Block)
+					{
+						XGameSaveProviderHandle NewProvider;
+						HRESULT hResult = XGameSaveInitializeProviderResult(Block, &NewProvider);
+						ProcessErrorCode(hResult, TEXT("InitializeProvider"));
+
+						// if the user was asked to resolve a save (i.e. choose between one offline and a newer online one for example), and selected "Cancel" or "Choose Next Time"
+						// hResult will be E_GS_USER_CANCELED *and* we'll have a valid Provider representing the old/offline one. It's not possible to close the Provider and return null
+						// because otherwise the next time this is called we will get S_OK and the same Provider, with no resolve UI displayed until the application is restarted.
+						PerUserSaveProviders.Add(User, NewProvider);
+
+						return hResult;
+					}
+				));
+			}
+		}
+	);
+}
+
+void FGDKSaveGameSystem::LoadGameIfExistsAsync(bool bAttemptToUseUI, const TCHAR* Name, FPlatformUserId PlatformUserId, FSaveGameAsyncLoadCompleteCallback Callback)
+{
+	// InternalLoadGameAsync already does everything that InternalDoesSaveGameExistAsync does but uses custom queuing, so just skip the first call
+	LoadGameAsync(bAttemptToUseUI, Name, PlatformUserId, Callback);
+}
+
+XGameSaveProviderHandle FGDKSaveGameSystem::GetProviderForUserId(FPlatformUserId PlatformUserId)
+{
+	if (!IGDKRuntimeModule::Get().IsAvailable())
+	{
+		return nullptr;
+	}
+
+	const FGDKUserHandle User = IGDKRuntimeModule::Get().GetUserHandleByPlatformId(PlatformUserId);
+	if (!User.IsValid())
+	{
+		return nullptr;
+	}
+
+	XGameSaveProviderHandle* HandlePtr = PerUserSaveProviders.Find(User);
+	if (HandlePtr == nullptr)
+	{
+		return nullptr;
+	}
+
+	return (*HandlePtr);
+}
+
+
+UE::Tasks::FTask FGDKSaveGameSystem::InternalGetSaveGameNamesAsync(FPlatformUserId PlatformUserId, TSharedRef<TArray<FString>> FoundSaves, FSaveGameAsyncGetNamesCallback Callback, TSharedPtr<bool> OutResult)
+{
+	// start the enumeration operation on a background thread
+	return AsyncTaskPipe.Launch(UE_SOURCE_LOCATION,
+		[this, PlatformUserId, FoundSaves, Callback, OutResult]()
+		{
+			bool bSucceeded = false;
+
+			// look up the provider first
+			XGameSaveProviderHandle Provider = GetProviderForUserId(PlatformUserId);
+			if (Provider != nullptr)
+			{
+				// enumerate the save game names
+				HRESULT hResult = S_OK;
+				if (GGDKUseXDKCompatibleSave == 0)
+				{
+					// GDK save game names are the container names themselves
+					auto GatherGDKContainerInfo = [](const XGameSaveContainerInfo* ContainerInfo, void* Context)
+					{
+						if (ContainerInfo->name != nullptr && (FCStringAnsi::Strcmp(ContainerInfo->name, XDKCompatibleContainerName) != 0)) // ignore XDK compatible saves when we're not looking for them
+						{
+							TArray<FString>* FoundSavesPtr = (TArray<FString>*)Context;
+							FoundSavesPtr->Add(FString(FUTF8ToTCHAR(ContainerInfo->name)));
+						}
+						return true;
+					};
+
+					hResult = XGameSaveEnumerateContainerInfo(Provider, &FoundSaves.Get(), GatherGDKContainerInfo);
+					ProcessErrorCode(hResult, TEXT("EnumerateContainerInfo"));
+				}
+				else if (DoesContainerExist(Provider, FString(XDKCompatibleContainerName)))
+				{
+					// XDK-compatible save game names are the blob names within the specifically-named container
+					auto GatherXDKBlobInfo = [](const XGameSaveBlobInfo* BlobInfo, void* Context)
+					{
+						if (BlobInfo->name != nullptr)
+						{
+							TArray<FString>* FoundSavesPtr = (TArray<FString>*)Context;
+							FoundSavesPtr->Add(FString(FUTF8ToTCHAR(BlobInfo->name)));
+						}
+						return true;
+					};
+
+					// create the XDK-compatible container & enumerate all blobs
+					XGameSaveContainerHandle Container = nullptr;
+					hResult = XGameSaveCreateContainer(Provider, XDKCompatibleContainerName, &Container);
+					ProcessErrorCode(hResult, TEXT("CreateContainer"));
+					if (Container)
+					{
+						hResult = XGameSaveEnumerateBlobInfo(Container, &FoundSaves.Get(), GatherXDKBlobInfo);
+						ProcessErrorCode(hResult, TEXT("EnumerateBlobInfo"));
+
+						XGameSaveCloseContainer(Container);
+					}
+
+				}
+				bSucceeded = SUCCEEDED(hResult);
+			}
+
+			// store the result, if requested
+			if (OutResult.IsValid())
+			{
+				*OutResult = bSucceeded;
+			}
+
+			// trigger the callback on the game thread
+			if (Callback)
+			{
+				OnAsyncComplete(
+					[PlatformUserId, bSucceeded, Callback, FoundSaves]()
+					{
+						Callback(PlatformUserId, bSucceeded, FoundSaves.Get());
+					}
+				);
+			}
+		},
+		// prerequisite - init save provider
+		InitSaveGameProviderAsync(PlatformUserId)
+	);
+}
+
+
+UE::Tasks::FTask FGDKSaveGameSystem::InternalDoesSaveGameExistAsync(const TCHAR* Name, FPlatformUserId PlatformUserId, FSaveGameAsyncExistsCallback Callback, TSharedPtr<ESaveExistsResult> OutResult)
+{
+	FString SlotName(Name);
+
+	// start the save operation on a background thread.
+	return AsyncTaskPipe.Launch(UE_SOURCE_LOCATION,		
+		[this, SlotName, PlatformUserId, Callback, OutResult]()
+		{
+			ESaveExistsResult Result = ESaveExistsResult::UnspecifiedError;
+
+			// look up the provider first. failure to find the provider is a different failure reason to 'does not exist'
+			XGameSaveProviderHandle Provider = GetProviderForUserId(PlatformUserId);
+			if (Provider != nullptr)
+			{
+				FString ContainerName, BlobName;
+				GetContainerAndBlobNames(ContainerName, BlobName, SlotName);
+				UE_LOG(LogGDKSaveGame, Display, TEXT("Looking for blob '%s' in container '%s'"), *BlobName, *ContainerName);
+
+				// enumerate the blobs in the container and see if the one we want exists
+				XGameSaveContainerHandle Container = nullptr;
+				HRESULT hResult = XGameSaveCreateContainer(Provider, TCHAR_TO_UTF8(*ContainerName), &Container);
+				ProcessErrorCode(hResult, TEXT("CreateContainer"));
+
+				Result = ISaveGameSystem::ESaveExistsResult::DoesNotExist;
+				if (Container)
+				{
+					uint32 DataSize = 0;
+					Result = GetBlobInfoByName(Container, BlobName, DataSize);
+					XGameSaveCloseContainer(Container);
+				}
+			}
+
+			// store the result, if requested
+			if (OutResult.IsValid())
+			{
+				*OutResult = Result;
+			}
+
+			// trigger the callback on the game thread
+			if (Callback)
+			{
+				OnAsyncComplete(
+					[SlotName, PlatformUserId, Result, Callback]()
+					{
+						Callback(SlotName, PlatformUserId, Result);
+					}
+				);
+			}
+		}, 
+		// prerequisite - init save provider
+		InitSaveGameProviderAsync(PlatformUserId)
+	);		
+}
+
+
+UE::Tasks::FTask FGDKSaveGameSystem::InternalSaveGameAsync(bool bAttemptToUseUI, const TCHAR* Name, FPlatformUserId PlatformUserId, TSharedRef<const TArray<uint8>> Data, FSaveGameAsyncOpCompleteCallback Callback, TSharedPtr<bool> OutResult)
+{
+	FString SlotName(Name);
+
+	// start the save operation on a background thread.
+	return AsyncTaskPipe.Launch(UE_SOURCE_LOCATION,
+		[this, SlotName, PlatformUserId, Data, Callback, OutResult]()
+		{
+			bool bSucceeded = false;
+
+			// look up the provider first
+			XGameSaveProviderHandle Provider = GetProviderForUserId(PlatformUserId);
+			if (Provider != nullptr)
+			{
+				FString ContainerName, BlobName;
+				GetContainerAndBlobNames(ContainerName, BlobName, SlotName);
+				UE_LOG(LogGDKSaveGame, Display, TEXT("Saving Blob '%s' in Container '%s'..."), *BlobName, *ContainerName);
+
+				XGameSaveContainerHandle Container = nullptr;
+				XGameSaveUpdateHandle Update = nullptr;
+
+				// create the container object (or open existing)
+				HRESULT hResult = XGameSaveCreateContainer(Provider, TCHAR_TO_UTF8(*ContainerName), &Container);
+				ProcessErrorCode(hResult, TEXT("CreateContainer"));
+
+				if (SUCCEEDED(hResult))
+				{
+					hResult = XGameSaveCreateUpdate(Container, TCHAR_TO_UTF8(*ContainerName), &Update);
+					ProcessErrorCode(hResult, TEXT("CreateUpdate"));
+				}
+				if (SUCCEEDED(hResult))
+				{
+					hResult = XGameSaveSubmitBlobWrite(Update, TCHAR_TO_UTF8(*BlobName), Data->GetData(), Data->Num());
+					ProcessErrorCode(hResult, TEXT("SubmitBlobWrite"));
+				}
+				if (SUCCEEDED(hResult))
+				{
+					hResult = XGameSaveSubmitUpdate(Update);
+					ProcessErrorCode(hResult, TEXT("SubmitUpdate"));
+				}
+
+				if (Update)
+				{
+					XGameSaveCloseUpdate(Update);
+				}
+				if (Container)
+				{
+					XGameSaveCloseContainer(Container);
+				}
+
+				bSucceeded = SUCCEEDED(hResult);
+			}
+
+			// store the result, if requested
+			if (OutResult.IsValid())
+			{
+				*OutResult = bSucceeded;
+			}
+
+			// trigger the callback on the game thread
+			if (Callback)
+			{
+				OnAsyncComplete(
+					[SlotName, PlatformUserId, bSucceeded, Callback]()
+					{
+						Callback(SlotName, PlatformUserId, bSucceeded);
+					}
+				);
+			}
+		},
+		// prerequisite - init save provider
+		InitSaveGameProviderAsync(PlatformUserId)
+	);
+}
+
+
+UE::Tasks::FTask FGDKSaveGameSystem::InternalLoadGameAsync(bool bAttemptToUseUI, const TCHAR* Name, FPlatformUserId PlatformUserId, TSharedRef<TArray<uint8>> Data, FSaveGameAsyncLoadCompleteCallback Callback, TSharedPtr<bool> OutResult)
+{
+	FString SlotName(Name);
+
+	// start the save operation on a background thread.
+	return AsyncTaskPipe.Launch(UE_SOURCE_LOCATION,
+		[this, SlotName, PlatformUserId, Data, Callback, OutResult]()
+		{
+			bool bSucceeded = false;
+
+			// look up the provider first
+			XGameSaveProviderHandle Provider = GetProviderForUserId(PlatformUserId);
+			if (Provider != nullptr)
+			{
+				FString ContainerName, BlobName;
+				GetContainerAndBlobNames(ContainerName, BlobName, SlotName);
+				UE_LOG(LogGDKSaveGame, Display, TEXT("Loading blob '%s' from container '%s'..."), *BlobName, *ContainerName);
+
+				if (DoesContainerExist(Provider, ContainerName))
+				{
+					// create the container object (or open existing)
+					XGameSaveContainerHandle Container = nullptr;
+					HRESULT hResult = XGameSaveCreateContainer(Provider, TCHAR_TO_UTF8(*ContainerName), &Container);
+					ProcessErrorCode(hResult, TEXT("CreateContainer"));
+
+					// read the blob
+					uint32 DataSize = 0;
+					if (Container != nullptr && GetBlobInfoByName(Container, BlobName, DataSize) == ISaveGameSystem::ESaveExistsResult::OK)
+					{
+						uint32 AllocationSize = DataSize;
+						AllocationSize += BlobName.Len() + 1; // +1 for null
+						AllocationSize += sizeof(XGameSaveBlob); // header
+
+						FTCHARToUTF8 BlobNameUTF8(*BlobName);
+						const char* BlobNames[] = { BlobNameUTF8.Get() };
+
+						uint32 BlobCount = 1;
+						XGameSaveBlob* Blob = (XGameSaveBlob*)FMemory::Malloc(AllocationSize);
+						hResult = XGameSaveReadBlobData(Container, BlobNames, &BlobCount, AllocationSize, Blob);
+						ProcessErrorCode(hResult, TEXT("ReadBlobData"));
+
+						UE_CLOG(BlobCount != 1, LogGDKSaveGame, Error, TEXT("Expected 1 blob, but got %d"), BlobCount);
+						if (SUCCEEDED(hResult) && BlobCount == 1)
+						{
+							int32 Offset = Data->AddUninitialized(Blob->info.size);
+							FMemory::Memcpy(Data->GetData() + Offset, Blob->data, Blob->info.size);
+							bSucceeded = true;
+						}
+
+						FMemory::Free(Blob);
+					}
+
+					if (Container)
+					{
+						XGameSaveCloseContainer(Container);
+					}
+				}
+			}
+
+			// store the result, if requested
+			if (OutResult.IsValid())
+			{
+				*OutResult = bSucceeded;
+			}
+
+			// trigger the callback on the game thread
+			if (Callback)
+			{
+				OnAsyncComplete(
+					[SlotName, PlatformUserId, bSucceeded, Callback, Data]()
+					{
+						Callback(SlotName, PlatformUserId, bSucceeded, Data.Get());
+					}
+				);
+			}
+		},
+		// prerequisite - init save provider
+		InitSaveGameProviderAsync(PlatformUserId)
+	);
+}
+
+
+UE::Tasks::FTask FGDKSaveGameSystem::InternalDeleteGameAsync(bool bAttemptToUseUI, const TCHAR* Name, FPlatformUserId PlatformUserId, FSaveGameAsyncOpCompleteCallback Callback, TSharedPtr<bool> OutResult)
+{
+	FString SlotName(Name);
+
+	// start the delete operation on a background thread.
+	return AsyncTaskPipe.Launch(UE_SOURCE_LOCATION,
+		[this, SlotName, PlatformUserId, Callback, OutResult]()
+		{
+			bool bSucceeded = false;
+
+			// look up the provider first
+			XGameSaveProviderHandle Provider = GetProviderForUserId(PlatformUserId);
+			if (Provider != nullptr)
+			{
+				FString ContainerName, BlobName;
+				GetContainerAndBlobNames(ContainerName, BlobName, SlotName);
+				UE_LOG(LogGDKSaveGame, Display, TEXT("Deleting blob '%s' from container '%s'..."), *BlobName, *ContainerName);
+
+				XGameSaveContainerHandle Container = nullptr;
+				XGameSaveUpdateHandle Update = nullptr;
+
+				if (DoesContainerExist(Provider, ContainerName))
+				{
+					// create the container object (or open existing)
+					HRESULT hResult = XGameSaveCreateContainer(Provider, TCHAR_TO_UTF8(*ContainerName), &Container);
+					ProcessErrorCode(hResult, TEXT("CreateContainer"));
+
+					uint32 DataSize = 0;
+					if (GetBlobInfoByName(Container, BlobName, DataSize) == ISaveGameSystem::ESaveExistsResult::OK)
+					{
+						if (SUCCEEDED(hResult))
+						{
+							hResult = XGameSaveCreateUpdate(Container, TCHAR_TO_UTF8(*ContainerName), &Update);
+							ProcessErrorCode(hResult, TEXT("CreateUpdate"));
+						}
+						if (SUCCEEDED(hResult))
+						{
+							hResult = XGameSaveSubmitBlobDelete(Update, TCHAR_TO_UTF8(*BlobName));
+							ProcessErrorCode(hResult, TEXT("SubmitBlobDelete"));
+						}
+						if (SUCCEEDED(hResult))
+						{
+							hResult = XGameSaveSubmitUpdate(Update);
+							ProcessErrorCode(hResult, TEXT("SubmitUpdate"));
+						}
+
+						if (Update)
+						{
+							XGameSaveCloseUpdate(Update);
+						}
+					}
+
+					uint32 NumBlobsRemaining = GetNumBlobsInContainer(Container);
+					if (Container)
+					{
+						XGameSaveCloseContainer(Container);
+					}
+
+					bSucceeded = SUCCEEDED(hResult);
+
+					// delete the container if it is now empty
+					if (NumBlobsRemaining == 0)
+					{
+						UE_LOG(LogGDKSaveGame, Display, TEXT("Deleting empty container '%s'..."), *ContainerName);
+						hResult = XGameSaveDeleteContainer(Provider, TCHAR_TO_UTF8(*ContainerName));
+						ProcessErrorCode(hResult, TEXT("DeleteContainer"));
+					}
+				}
+			}
+
+
+			// store the result, if requested
+			if (OutResult.IsValid())
+			{
+				*OutResult = bSucceeded;
+			}
+
+
+			// trigger the callback on the game thread
+			if (Callback)
+			{
+				OnAsyncComplete(
+					[SlotName, PlatformUserId, bSucceeded, Callback]()
+					{
+						Callback(SlotName, PlatformUserId, bSucceeded);
+					}
+				);
+			}
+		},
+		// prerequisite - init save provider
+		InitSaveGameProviderAsync(PlatformUserId)
+	);
+}
+
+
+
+void FGDKSaveGameSystem::InitAsync(bool bAttemptToUseUI, FPlatformUserId PlatformUserId, FSaveGameAsyncInitCompleteCallback Callback)
+{
+	// start the init operation on a background thread
+	AsyncTaskPipe.Launch(UE_SOURCE_LOCATION,
+		[this, PlatformUserId, Callback]
+		{
+			// get the save provider handle
+			XGameSaveProviderHandle Provider = GetProviderForUserId(PlatformUserId);
+			bool bSuccess = (Provider != nullptr);
+
+			//trigger the callback on the game thread
+			if (Callback)
+			{
+				OnAsyncComplete(
+					[PlatformUserId, bSuccess, Callback]()
+					{
+						Callback(PlatformUserId, bSuccess);
+					}
+				);
+			}
+		},
+		// prerequisite - init save provider
+		InitSaveGameProviderAsync(PlatformUserId)
+	);
+}
+
+
+
+void FGDKSaveGameSystem::WaitForAsyncTask(UE::Tasks::FTask AsyncSaveTask)
+{
+	// need to pump messages on the game thread
+	if (IsInGameThread())
+	{
+		// Suspend the hang and hitch heartbeats, as this is a long running task.
+		FSlowHeartBeatScope SuspendHeartBeat;
+		FDisableHitchDetectorScope SuspendGameThreadHitch;
+
+		while (!AsyncSaveTask.IsCompleted())
+		{
+			FPlatformMisc::PumpMessagesOutsideMainLoop();
+		}
+	}
+	else
+	{
+		// not running on the game thread, so just block until the async operation comes back
+		AsyncSaveTask.Wait();
+	}
+}
+
+
+
+void FGDKSaveGameSystem::GetContainerAndBlobNames( FString& OutContainerName, FString& OutBlobName, const FString& InSaveName )
+{
+	if (GGDKUseXDKCompatibleSave)
+	{
+		OutContainerName = XDKCompatibleContainerName;
+		OutBlobName = InSaveName;
+	}
+	else
+	{
+		OutContainerName = SanitizeContainerName(*InSaveName);
+		OutBlobName = DefaultBlobName;
+	}
+}
+
+
+FString FGDKSaveGameSystem::SanitizeContainerName(const TCHAR* Name)
+{
+	TCHAR ContainerName[MaxContainerName];
+	int WriteIdx = 0;
+	while( WriteIdx < (MaxContainerName-1) && *Name != '\0')
+	{
+		if (FChar::IsAlnum(*Name) || *Name == '_' || *Name == '-')
+		{
+			ContainerName[WriteIdx++] = *Name;
+		}
+		Name++;
+	}
+	ContainerName[WriteIdx++] = '\0';
+
+	return (ContainerName[0] == '\0') ? TEXT("Default") : ContainerName;
+}
+
+
+
+ISaveGameSystem::ESaveExistsResult FGDKSaveGameSystem::GetBlobInfoByName( XGameSaveContainerHandle InContainer, const FString& InBlobName, uint32& OutDataSize )
+{
+	// helper for enumerating blobs
+	struct FBlobEnumData
+	{
+		FBlobEnumData( const FString& InName, XGameSaveContainerHandle InContainer ) :
+			Container(InContainer),
+			Name(*InName),
+			Size(0),
+			bFound(false)
+		{}
+
+		XGameSaveContainerHandle Container;
+		FTCHARToUTF8 Name;
+		uint32 Size;
+		bool bFound;
+
+		void Enumerate()
+		{
+			HRESULT hResult = XGameSaveEnumerateBlobInfo( Container, this, []( const XGameSaveBlobInfo* Info, void* Context )
+			{
+				FBlobEnumData* BlobEnumData = (FBlobEnumData*)Context;
+				if (FCStringAnsi::Strcmp(Info->name, BlobEnumData->Name.Get()) == 0) //NB. blob names are case-sensitive
+				{
+					BlobEnumData->Size = Info->size;
+					BlobEnumData->bFound = true;
+					return false;
+				}
+				return true;
+			});
+
+			FGDKSaveGameSystem::ProcessErrorCode(hResult, TEXT("EnumerateBlobInfo"));
+		}
+	};
+
+	// enumerate the blobs and find the size of the one we are looking for
+	FBlobEnumData BlobEnumData(InBlobName, InContainer);
+
+	BlobEnumData.Enumerate();
+
+	if (BlobEnumData.bFound)
+	{
+		UE_LOG(LogGDKSaveGame, Display, TEXT("Found blob '%s' - %d bytes"), *InBlobName, BlobEnumData.Size );
+
+		OutDataSize = BlobEnumData.Size;
+		return ISaveGameSystem::ESaveExistsResult::OK;
+	}
+	else
+	{
+		UE_LOG(LogGDKSaveGame, Display, TEXT("Blob '%s' not found"), *InBlobName );
+
+		OutDataSize = 0;
+		return ISaveGameSystem::ESaveExistsResult::DoesNotExist;
+	}
+}
+
+
+uint32 FGDKSaveGameSystem::GetNumBlobsInContainer( XGameSaveContainerHandle InContainer )
+{
+	uint32 Result = 0;
+	HRESULT hResult = XGameSaveEnumerateBlobInfo( InContainer, &Result, []( const XGameSaveBlobInfo* Info, void* Context )
+	{
+		(*static_cast<uint32*>(Context))++;
+		return true;
+	});
+	ProcessErrorCode(hResult, TEXT("EnumerateBlobInfo"));
+
+	return Result;
+}
+
+
+bool FGDKSaveGameSystem::DoesContainerExist( XGameSaveProviderHandle InProvider, const FString& InContainerName )
+{
+	bool bContainerExists = false;
+	HRESULT hResult = XGameSaveGetContainerInfo( InProvider, TCHAR_TO_UTF8(*InContainerName), &bContainerExists, [](const XGameSaveContainerInfo* Container, void* Context) //NB. container names are case sensitive
+	{
+		(*static_cast<bool*>(Context)) = true;
+		return true;
+	});
+	ProcessErrorCode(hResult, TEXT("ContainerInfo") );
+
+	UE_CLOG(!bContainerExists, LogGDKSaveGame, Log, TEXT("Container '%s' not found"), *InContainerName );
+	return bContainerExists;
+}
+
+
+void FGDKSaveGameSystem::ProcessErrorCode( HRESULT Result, const TCHAR* Msg )
+{
+	switch(Result)
+	{
+		case S_OK:                       return;
+		case E_GS_NO_ACCESS:             UE_LOG(LogGDKSaveGame, Error, TEXT("%s: E_GS_NO_ACCESS: SCID, TitleID or package identity may not be set up correctly."), Msg); break;
+		case E_GS_OUT_OF_LOCAL_STORAGE:  UE_LOG(LogGDKSaveGame, Error, TEXT("%s: E_GS_OUT_OF_LOCAL_STORAGE."), Msg); break;
+		case E_GS_UPDATE_TOO_BIG:        UE_LOG(LogGDKSaveGame, Error, TEXT("%s: E_GS_UPDATE_TOO_BIG: reduce the savegame size"), Msg); break;
+		case E_GS_QUOTA_EXCEEDED:        UE_LOG(LogGDKSaveGame, Error, TEXT("%s: E_GS_QUOTA_EXCEEDED: reduce the savegame size"), Msg); break;
+		case E_GS_CONTAINER_NOT_IN_SYNC: UE_LOG(LogGDKSaveGame, Error, TEXT("%s: E_GS_CONTAINER_NOT_IN_SYNC"), Msg); break;
+		case E_GS_CONTAINER_SYNC_FAILED: UE_LOG(LogGDKSaveGame, Error, TEXT("%s: E_GS_CONTAINER_SYNC_FAILED"), Msg); break;
+		case E_GS_HANDLE_EXPIRED:        UE_LOG(LogGDKSaveGame, Error, TEXT("%s: E_GS_HANDLE_EXPIRED"), Msg); break;
+		case E_GS_USER_CANCELED:         UE_LOG(LogGDKSaveGame, Error, TEXT("%s: E_GS_USER_CANCELED"), Msg); break;
+		default:                         UE_LOG(LogGDKSaveGame, Error, TEXT("%s: Failed 0x%X"), Msg, Result ); break;
+	}
+}
+
+void FGDKSaveGameSystem::DebugListSavesForAllUsers()
+{
+	TArray<FGDKUserHandle> AllUserHandles = IGDKRuntimeModule::Get().GetAllUserHandles();
+	for (FGDKUserHandle UserHandle : AllUserHandles )
+	{
+		const FPlatformUserId PlatformUserId = IGDKRuntimeModule::Get().GetPlatformIdByUserHandle(UserHandle);
+		const int32 UserIndex = FPlatformMisc::GetUserIndexForPlatformUser(PlatformUserId);
+		DebugListSavesForUser(UserIndex);
+	}
+}
+
+void FGDKSaveGameSystem::DebugListSavesForUser( const int32 UserIndex )
+{
+	const FPlatformUserId PlatformUserId = FPlatformMisc::GetPlatformUserForUserIndex(UserIndex);
+	XGameSaveProviderHandle Provider = GetProviderForUserId(PlatformUserId);
+	if (Provider == nullptr)
+	{
+		return;
+	}
+
+	GDK_SCOPE_NOT_TIME_SENSITIVE(); // (debug only) XGameSaveEnumerateContainerInfo & XGameSaveEnumerateBlobInfo are not safe to call on a time-sensitive thread
+
+	auto DumpContainerInfo = []( const XGameSaveContainerInfo* Info, void* Context )
+	{
+		bool bIsXDKCompatibleContainer = Info->name && (FCStringAnsi::Strcmp( Info->name, XDKCompatibleContainerName ) == 0);
+		FDateTime ModifiedTime = FDateTime::FromUnixTimestamp(Info->lastModifiedTime);
+		FString ModifiedAgo = (FDateTime::UtcNow() - ModifiedTime).ToString(TEXT("%dd %hh %mmin %ss ago"));
+		ModifiedAgo.RemoveFromStart(TEXT("+"));
+
+		UE_LOG(LogGDKSaveGame, Display, TEXT("\t%sContainer: '%s' %d bytes (%s)"), bIsXDKCompatibleContainer ? TEXT("XDK-Compatible ") : TEXT(""), Info->name ? UTF8_TO_TCHAR(Info->name) : TEXT("<null>"), Info->totalSize, Info->needsSync ? TEXT("needs sync") : TEXT("synced"));
+		if (Info->name != nullptr)
+		{
+			// create this container
+			XGameSaveContainerHandle Container = nullptr;
+			HRESULT hResult = XGameSaveCreateContainer( (XGameSaveProviderHandle)Context, Info->name, &Container );
+			ProcessErrorCode( hResult, *FString::Printf( TEXT("cannot create save game container for %s"), UTF8_TO_TCHAR(Info->name) ) );
+
+			// enumerate all blobs
+			if (Container)
+			{
+				auto DumpBlobInfo = []( const XGameSaveBlobInfo* Info, void* Context )
+				{
+					UE_LOG( LogGDKSaveGame, Display, TEXT("\t\tBlob: '%s' %d bytes"), Info->name ? UTF8_TO_TCHAR(Info->name) : TEXT("<null>"), Info->size );
+					return true;
+				};
+
+				hResult = XGameSaveEnumerateBlobInfo( Container, nullptr, DumpBlobInfo );
+				ProcessErrorCode( hResult, *FString::Printf( TEXT("cannot enumerate save games for container '%s'"), UTF8_TO_TCHAR(Info->name) ) );
+				XGameSaveCloseContainer(Container);
+			}
+		}
+		UE_LOG(LogGDKSaveGame, Display, TEXT("\tLast modified: %s  (%s)"), *ModifiedAgo, *ModifiedTime.ToString());
+		UE_LOG(LogGDKSaveGame, Display, TEXT("") );
+
+		return true;
+	};
+
+	// enumerate all containers
+	UE_LOG(LogGDKSaveGame, Display, TEXT("*** User %d save data (%s) ***"), UserIndex, *LexToString(IGDKRuntimeModule::Get().GetUserHandleByPlatformId(PlatformUserId)) );
+	UE_LOG(LogGDKSaveGame, Display, TEXT("") );
+	HRESULT hResult = XGameSaveEnumerateContainerInfo( Provider, (void*)Provider, DumpContainerInfo );
+	ProcessErrorCode( hResult, TEXT("cannot enumerate save game containers") );
+	UE_LOG(LogGDKSaveGame, Display, TEXT("********") );
+	UE_LOG(LogGDKSaveGame, Display, TEXT("") );
+}
+
+
+
+#endif //WITH_GRDK
diff --git a/Engine/Source/Runtime/Microsoft/GDK/GDKSaveGameSystem/Private/GDKSaveGameSystemModule.cpp b/Engine/Source/Runtime/Microsoft/GDK/GDKSaveGameSystem/Private/GDKSaveGameSystemModule.cpp
new file mode 100644
--- /dev/null
+++ b/Engine/Source/Runtime/Microsoft/GDK/GDKSaveGameSystem/Private/GDKSaveGameSystemModule.cpp
@@ -0,0 +1,188 @@
+// Copyright Epic Games, Inc. All Rights Reserved.
+
+#include "Modules/ModuleInterface.h"
+#include "Modules/ModuleManager.h"
+#include "GDKRuntimeModule.h"
+#include "GDKSaveGameSystem.h"
+
+class FGDKSaveGameSystemModule : public ISaveGameSystemModule
+{
+public:
+#if WITH_GRDK
+	std::atomic<bool> bWantSystemCreation;
+	TUniquePtr<FGDKSaveGameSystem> SaveGameSystem;
+
+	FGDKSaveGameSystemModule()
+		: bWantSystemCreation(true)
+	{
+#if WITH_EDITOR
+		// tear down the savegame system when PIE finishes
+		IGDKRuntimeModule::Get().GetOnTeardownForPIE().AddRaw( this, &FGDKSaveGameSystemModule::OnTeardownForPIE );
+#else
+		// prevent creation of GDK savegame system if the GDK runtime is unavailable
+		if (!IGDKRuntimeModule::Get().IsAvailable())
+		{
+			bWantSystemCreation = false;
+		}
+#endif
+	}
+
+	virtual void ShutdownModule() override
+	{
+		SaveGameSystem.Reset();
+
+#if WITH_EDITOR
+		if (IGDKRuntimeModule* GDKRuntime = IGDKRuntimeModule::TryGet())
+		{
+			GDKRuntime->GetOnTeardownForPIE().RemoveAll(this);
+		}
+#endif
+	}
+
+	virtual ISaveGameSystem* GetSaveGameSystem() override
+	{
+#if WITH_EDITOR
+		if (!IGDKRuntimeModule::Get().IsAvailable())
+		{
+			return nullptr;
+		}
+#endif
+
+		if (bWantSystemCreation)
+		{
+			bWantSystemCreation = false;
+			SaveGameSystem.Reset( new FGDKSaveGameSystem());
+		}
+
+		return SaveGameSystem.Get();
+	}
+
+
+	static inline FGDKSaveGameSystem* GetStatic()
+	{
+		FGDKSaveGameSystemModule& ThisModule = FModuleManager::LoadModuleChecked<FGDKSaveGameSystemModule>("GDKSaveGameSystem");
+		return (FGDKSaveGameSystem*)ThisModule.GetSaveGameSystem();
+	}
+
+#if WITH_EDITOR
+	void OnTeardownForPIE()
+	{
+		SaveGameSystem.Reset();
+		bWantSystemCreation = true;
+	}
+#endif
+
+
+
+#else
+
+	// dummy
+	virtual ISaveGameSystem* GetSaveGameSystem() override
+	{
+		return nullptr;
+	}
+#endif //WITH_GRDK
+};
+
+
+#if WITH_GRDK
+
+static FAutoConsoleCommand CmdGDKListSaves(
+	TEXT("GDK.SaveGame.List"),
+	TEXT("[optional UserIndex] Print out users' save data"),
+	FConsoleCommandWithArgsDelegate::CreateLambda( [](const TArray<FString>& Args )
+	{
+		if (Args.Num() > 0)
+		{
+			int32 UserIndex = Args.Num() > 0 ? FCString::Atoi(*Args[0]) : 0;
+			FGDKSaveGameSystemModule::GetStatic()->DebugListSavesForUser(UserIndex);
+		}
+		else
+		{
+			FGDKSaveGameSystemModule::GetStatic()->DebugListSavesForAllUsers();
+		}
+	})
+);
+
+static FAutoConsoleCommand CmdGDKDoesSaveExist(
+	TEXT("GDK.SaveGame.Exists"),
+	TEXT("[Name] [optional UserIndex] - Checks to see if a given save exists"),
+	FConsoleCommandWithArgsDelegate::CreateLambda( [](const TArray<FString>& Args )
+	{
+		if (Args.Num() > 0)
+		{
+			int32 UserIndex = Args.Num() > 1 ? FCString::Atoi(*Args[1]) : 0;
+			bool bExists = FGDKSaveGameSystemModule::GetStatic()->DoesSaveGameExist(*Args[0], UserIndex );
+			UE_LOG( LogGDKSaveGame, Log, TEXT("Savegame %s %s"), *Args[0], bExists ? TEXT("exists") : TEXT("does not exist") );
+		}
+	})
+);
+
+static FAutoConsoleCommand CmdGDKDeleteSave(
+	TEXT("GDK.SaveGame.Delete"),
+	TEXT("[Name] [optional UserIndex] - Deletes the given save"),
+	FConsoleCommandWithArgsDelegate::CreateLambda( [](const TArray<FString>& Args )
+	{
+		if (Args.Num() > 0)
+		{
+			int32 UserIndex = Args.Num() > 1 ? FCString::Atoi(*Args[1]) : 0;
+			bool bSuccess = FGDKSaveGameSystemModule::GetStatic()->DeleteGame(false, *Args[0], UserIndex );
+			UE_LOG( LogGDKSaveGame, Log, TEXT("Savegame %s was %s"), *Args[0], bSuccess ? TEXT("deleted") : TEXT("not deleted") );
+			FGDKSaveGameSystemModule::GetStatic()->DebugListSavesForUser(UserIndex);
+		}
+	})
+);
+
+static FAutoConsoleCommand CmdGDKLoadBlob(
+	TEXT("GDK.SaveGame.DummyLoad"),
+	TEXT("[Name] [optional UserIndex] - Tries to load the given save data. note: deletes the data once loaded"),
+	FConsoleCommandWithArgsDelegate::CreateLambda( [](const TArray<FString>& Args )
+	{
+		if (Args.Num() > 0)
+		{
+			int32 UserIndex = Args.Num() > 1 ? FCString::Atoi(*Args[1]) : 0;
+			TArray<uint8> Result;
+			bool bSuccess = FGDKSaveGameSystemModule::GetStatic()->LoadGame(false, *Args[0], UserIndex, Result );
+			UE_LOG( LogGDKSaveGame, Log, TEXT("Savegame %s was %s"), *Args[0], bSuccess ? TEXT("loaded") : TEXT("not loaded") );
+			FGDKSaveGameSystemModule::GetStatic()->DebugListSavesForUser(UserIndex);
+		}
+	})
+);
+
+static FAutoConsoleCommand CmdGDKSaveBlob(
+	TEXT("GDK.SaveGame.DummySave"),
+	TEXT("[Name] [optional UserIndex] - Tries to save dummy data with the given name"),
+	FConsoleCommandWithArgsDelegate::CreateLambda( [](const TArray<FString>& Args )
+	{
+		if (Args.Num() > 0)
+		{
+			int32 UserIndex = Args.Num() > 1 ? FCString::Atoi(*Args[1]) : 0;
+			TArray<uint8> Data;
+			Data.AddZeroed(256);
+			bool bSuccess = FGDKSaveGameSystemModule::GetStatic()->SaveGame(false, *Args[0], UserIndex, Data );
+			UE_LOG( LogGDKSaveGame, Log, TEXT("Savegame %s was %s"), *Args[0], bSuccess ? TEXT("saved") : TEXT("not saved") );
+			FGDKSaveGameSystemModule::GetStatic()->DebugListSavesForUser(UserIndex);
+		}
+	})
+);
+
+static FAutoConsoleCommand CmdGDKInitSave(
+	TEXT("GDK.SaveGame.Init"),
+	TEXT("[optional UserIndex] initialise for this user"),
+	FConsoleCommandWithArgsDelegate::CreateLambda([](const TArray<FString>& Args)
+	{
+		int32 UserIndex = Args.Num() > 0 ? FCString::Atoi(*Args[0]) : 0;
+		FGDKSaveGameSystemModule::GetStatic()->InitAsync(false, FPlatformMisc::GetPlatformUserForUserIndex(UserIndex), [](FPlatformUserId id, bool succ)
+		{
+			UE_LOG(LogGDKSaveGame, Log, TEXT("Savegame init %d %s"), id.GetInternalId(), succ ? TEXT("true") : TEXT("false"));
+		});
+	})
+);
+
+#endif //WITH_GRDK
+
+
+
+
+
+IMPLEMENT_MODULE(FGDKSaveGameSystemModule, GDKSaveGameSystem);
diff --git a/Engine/Source/Runtime/Microsoft/GDK/GDKSaveGameSystem/Public/GDKSaveGameSystem.h b/Engine/Source/Runtime/Microsoft/GDK/GDKSaveGameSystem/Public/GDKSaveGameSystem.h
new file mode 100644
--- /dev/null
+++ b/Engine/Source/Runtime/Microsoft/GDK/GDKSaveGameSystem/Public/GDKSaveGameSystem.h
@@ -0,0 +1,95 @@
+// Copyright Epic Games, Inc. All Rights Reserved.
+
+#pragma once
+
+#include "SaveGameSystem.h"
+#if WITH_GRDK
+#include "GDKHandle.h"
+
+#include "Microsoft/AllowMicrosoftPlatformTypes.h"
+THIRD_PARTY_INCLUDES_START
+#include <XGameSave.h>
+#include <XTaskQueue.h>
+THIRD_PARTY_INCLUDES_END
+#include "Microsoft/HideMicrosoftPlatformTypes.h"
+
+#define UE_API GDKSAVEGAMESYSTEM_API
+
+DECLARE_LOG_CATEGORY_EXTERN(LogGDKSaveGame, Log, All);
+
+class FGDKSaveGameSystem : public FBaseAsyncSaveGameSystem
+{
+public:
+	UE_API FGDKSaveGameSystem();
+	UE_API virtual ~FGDKSaveGameSystem();
+
+	/** Returns true if the platform has a native UI (like many consoles) */
+	virtual bool PlatformHasNativeUI() override
+	{
+		return false;
+	}
+
+	virtual bool DoesSaveSystemSupportMultipleUsers() override
+	{
+		return true;
+	}
+
+	/* (Optional) initialise the save system for the given user. GDK may display the 'syncing' or 'resolve save' dialog when creating the save provider */
+	UE_API virtual void InitAsync(bool bAttemptToUseUI, FPlatformUserId PlatformUserId, FSaveGameAsyncInitCompleteCallback Callback) override;
+	
+	/** Override as the custom queuing breaks the default implementation and is unnecessary */
+	UE_API virtual void LoadGameIfExistsAsync(bool bAttemptToUseUI, const TCHAR* Name, FPlatformUserId PlatformUserId, FSaveGameAsyncLoadCompleteCallback Callback) override;
+
+
+	/** Debugging helper functions */
+	UE_API void DebugListSavesForAllUsers();
+	UE_API void DebugListSavesForUser( const int32 UserIndex );
+
+private:
+	static UE_API const uint32 MaxContainerName;
+	static UE_API const char* DefaultBlobName;
+	static UE_API const char* XDKCompatibleContainerName;
+
+	// internal async helpers
+	UE_API virtual UE::Tasks::FTask InternalDoesSaveGameExistAsync(const TCHAR* Name, FPlatformUserId PlatformUserId, FSaveGameAsyncExistsCallback Callback, TSharedPtr<ESaveExistsResult> OutResult = nullptr) override;
+	UE_API virtual UE::Tasks::FTask InternalSaveGameAsync(bool bAttemptToUseUI, const TCHAR* Name, FPlatformUserId PlatformUserId, TSharedRef<const TArray<uint8>> Data, FSaveGameAsyncOpCompleteCallback Callback, TSharedPtr<bool> OutResult = nullptr) override;
+	UE_API virtual UE::Tasks::FTask InternalLoadGameAsync(bool bAttemptToUseUI, const TCHAR* Name, FPlatformUserId PlatformUserId, TSharedRef<TArray<uint8>> Data, FSaveGameAsyncLoadCompleteCallback Callback, TSharedPtr<bool> OutResult = nullptr) override;
+	UE_API virtual UE::Tasks::FTask InternalDeleteGameAsync(bool bAttemptToUseUI, const TCHAR* Name, FPlatformUserId PlatformUserId, FSaveGameAsyncOpCompleteCallback Callback, TSharedPtr<bool> OutResult = nullptr) override;
+	UE_API virtual UE::Tasks::FTask InternalGetSaveGameNamesAsync(FPlatformUserId PlatformUserId, TSharedRef<TArray<FString>> FoundSaves, FSaveGameAsyncGetNamesCallback Callback, TSharedPtr<bool> OutResult = nullptr) override;
+
+	UE_API UE::Tasks::FTask InitSaveGameProviderAsync(FPlatformUserId PlatformUserId);
+
+	UE_API virtual void WaitForAsyncTask(UE::Tasks::FTask AsyncSaveTask) override;
+
+	// user sign in/out callback
+	XTaskQueueRegistrationToken UserChangeCallbackToken;
+	UE_API void OnUserChange(XUserLocalId UserLocalId, XUserChangeEvent Event);
+
+	// get the final blob and container names
+	static UE_API void GetContainerAndBlobNames( FString& OutContainerName, FString& OutBlobName, const FString& InSaveName );
+
+	// remove invalid characters from the given save container name
+	static UE_API FString SanitizeContainerName(const TCHAR* Name);
+
+	// get the info for a particular named blob
+	static UE_API ESaveExistsResult GetBlobInfoByName( XGameSaveContainerHandle InContainer, const FString& InBlobName, uint32& OutDataSize );
+
+	// get all blobs in the given container
+	static UE_API uint32 GetNumBlobsInContainer( XGameSaveContainerHandle InContainer );
+
+	// check whether the given container exists
+	static UE_API bool DoesContainerExist( XGameSaveProviderHandle InProvider, const FString& InContainerName );
+	
+	// error display helper
+	static UE_API void ProcessErrorCode( HRESULT Result, const TCHAR* Msg );
+
+	// get a provider for the user's save provider
+	UE_API XGameSaveProviderHandle GetProviderForUserId(FPlatformUserId UserId);
+
+	// cache the save providers for each user
+	TMap<FGDKUserHandle, XGameSaveProviderHandle> PerUserSaveProviders;
+};
+
+#undef UE_API
+
+#endif //WITH_GRDK
diff --git a/Engine/Source/Runtime/Microsoft/GDK/MSGameStorePlatformFeatures/MSGameStorePlatformFeatures.Build.cs b/Engine/Source/Runtime/Microsoft/GDK/MSGameStorePlatformFeatures/MSGameStorePlatformFeatures.Build.cs
new file mode 100644
--- /dev/null
+++ b/Engine/Source/Runtime/Microsoft/GDK/MSGameStorePlatformFeatures/MSGameStorePlatformFeatures.Build.cs
@@ -0,0 +1,24 @@
+// Copyright Epic Games, Inc. All Rights Reserved.
+
+using UnrealBuildTool;
+
+[SupportedPlatformGroups("Windows")]
+public class MSGameStorePlatformFeatures : ModuleRules
+{
+	public MSGameStorePlatformFeatures( ReadOnlyTargetRules Target) : base(Target)
+	{
+		BinariesSubFolder = "GDK";
+
+		PrivateDependencyModuleNames.AddRange(new string[]
+			{
+				"Core",
+				"Engine",
+				"GRDK",
+			});
+
+		PublicDependencyModuleNames.AddRange(new string[]
+			{
+				"WindowsPlatformFeatures",
+			});
+	}
+}
diff --git a/Engine/Source/Runtime/Microsoft/GDK/MSGameStorePlatformFeatures/Private/MSGameStorePlatformFeatures.cpp b/Engine/Source/Runtime/Microsoft/GDK/MSGameStorePlatformFeatures/Private/MSGameStorePlatformFeatures.cpp
new file mode 100644
--- /dev/null
+++ b/Engine/Source/Runtime/Microsoft/GDK/MSGameStorePlatformFeatures/Private/MSGameStorePlatformFeatures.cpp
@@ -0,0 +1,37 @@
+// Copyright Epic Games, Inc. All Rights Reserved.
+
+#include "MSGameStorePlatformFeatures.h"
+#if WITH_GRDK
+#include "Microsoft/AllowMicrosoftPlatformTypes.h"
+#include <XGame.h>
+#include <XAppCapture.h>
+#include "Microsoft/HideMicrosoftPlatformTypes.h"
+
+void FMSGameStorePlatformFeaturesModule::SetScreenshotEnableState(bool bEnabled)
+{
+	if (bEnabled)
+	{
+		XAppCaptureEnableRecord();
+	}
+	else
+	{
+		XAppCaptureDisableRecord();
+	}
+}
+
+
+FString FMSGameStorePlatformFeaturesModule::GetUniqueAppId()
+{
+	uint32_t TitleId = 0;
+	if (SUCCEEDED(XGameGetXboxTitleId(&TitleId)))
+	{
+		return FString::Printf(TEXT("%X"), TitleId);
+	}
+	else
+	{
+		return FString();
+	}
+}
+#endif //WITH_GRDK
+
+IMPLEMENT_MODULE(FMSGameStorePlatformFeaturesModule, MSGameStorePlatformFeatures);
diff --git a/Engine/Source/Runtime/Microsoft/GDK/MSGameStorePlatformFeatures/Public/MSGameStorePlatformFeatures.h b/Engine/Source/Runtime/Microsoft/GDK/MSGameStorePlatformFeatures/Public/MSGameStorePlatformFeatures.h
new file mode 100644
--- /dev/null
+++ b/Engine/Source/Runtime/Microsoft/GDK/MSGameStorePlatformFeatures/Public/MSGameStorePlatformFeatures.h
@@ -0,0 +1,13 @@
+// Copyright Epic Games, Inc. All Rights Reserved.
+
+#include "WindowsPlatformFeatures.h"
+
+class FMSGameStorePlatformFeaturesModule : public FWindowsPlatformFeaturesModule
+{
+public:
+#if WITH_GRDK
+	virtual FString GetUniqueAppId() override;
+	virtual void SetScreenshotEnableState(bool bEnabled) override;
+#endif
+};
+
