fix for GDKRuntimeModule not loading the correct ini settings unless the MSGamingRuntime module was part of the project.

#rnx

#Synced-CL 50857758

#Synced-CL 50857758

diff --git a/Engine/Source/Runtime/Microsoft/GDKRuntime/Private/GDKRuntimeModule.cpp b/Engine/Source/Runtime/Microsoft/GDKRuntime/Private/GDKRuntimeModule.cpp
--- a/Engine/Source/Runtime/Microsoft/GDKRuntime/Private/GDKRuntimeModule.cpp
+++ b/Engine/Source/Runtime/Microsoft/GDKRuntime/Private/GDKRuntimeModule.cpp
@@ -164,7 +164,7 @@ private:
 	void Teardown();
 	void CacheConfiguration();
 	bool bGDKEnvironmentInitialized = false;
-	bool bHasMSGamingRuntimeModule = false;
+	bool bHasMSGamingModules = false;
 
 	TUniquePtr<FGDKUserManager> UserManagerInstance;
 
@@ -249,10 +249,12 @@ static FAutoConsoleCommand CmdGDKDebugAddUser(
 void FGDKRuntimeModule::Init()
 {
 #if PLATFORM_WINDOWS
-	// cache whether the MSGamingRuntime module is loaded. This determines where the GDK target settings are read from
-	bHasMSGamingRuntimeModule = FModuleManager::Get().ModuleExists(TEXT("MSGamingRuntime"));
+	// cache whether any of the MSGaming modules are part of the project. This determines where the GDK target settings are read from
+	bHasMSGamingModules = FModuleManager::Get().ModuleExists(TEXT("MSGamingRuntime")) ||
+							FModuleManager::Get().ModuleExists(TEXT("MSGamingSupport")) ||
+							FModuleManager::Get().ModuleExists(TEXT("MSGameStore"));
 #else
-	bHasMSGamingRuntimeModule = false;
+	bHasMSGamingModules = false;
 #endif
 
 	// initialize GDK runtime
@@ -740,7 +742,7 @@ const TCHAR* FGDKRuntimeModule::GetConfigSectionName() const
 #if PLATFORM_WINDOWS
 	static FString ConfigSectionName = [this]()
 	{
-		return bHasMSGamingRuntimeModule ? TEXT("/Script/MSGamingSupport.MSGamingSettings") : FPlatformProperties::GetRuntimeSettingsClassName();
+		return bHasMSGamingModules ? TEXT("/Script/MSGamingSupport.MSGamingSettings") : FPlatformProperties::GetRuntimeSettingsClassName();
 	}();
 
 	return *ConfigSectionName;
