fix for WinGDK dev inititialization GDK runtime
- on-the-fly MSGame.config generation requires GConfig

#rnx

#Synced-CL 51007224

diff --git a/Engine/Source/Runtime/Microsoft/GDKRuntime/Private/GDKRuntimeInitialize.cpp b/Engine/Source/Runtime/Microsoft/GDKRuntime/Private/GDKRuntimeInitialize.cpp
--- a/Engine/Source/Runtime/Microsoft/GDKRuntime/Private/GDKRuntimeInitialize.cpp
+++ b/Engine/Source/Runtime/Microsoft/GDKRuntime/Private/GDKRuntimeInitialize.cpp
@@ -21,9 +21,6 @@ THIRD_PARTY_INCLUDES_START
 THIRD_PARTY_INCLUDES_END
 #include "Microsoft/HideMicrosoftPlatformTypes.h"
 
-#ifndef WITH_GRDK_DEV_RUNTIME_INIT
-	#define WITH_GRDK_DEV_RUNTIME_INIT PLATFORM_WINDOWS && !UE_BUILD_SHIPPING && (_GRDK_EDITION >= 240600)
-#endif
 
 #if WITH_GRDK_DEV_RUNTIME_INIT
 static HRESULT InitXGameRuntimeForDevelopment(const TCHAR* ConfigIniSection);
@@ -154,7 +151,6 @@ static FAnsiString FindMSGameConfigFile()
 	// check the saved folder
 #if !WITH_EDITOR
 	// this is primarily to support F5 in Visual Studio without specifying a -basedir ... in this case the game is loading cooked content from the Saved/Cooked directory or via Zen streaming
-	// it will most likely be necessary to specify -plugin=<path to>\Engine\Platforms\GDK\Plugins on the command line while the GDK runtime plugin is under an NDA platform folder.
 	if (FPaths::IsProjectFilePathSet())
 	{
 		MSGameConfigFile = FPaths::Combine( FPaths::ProjectSavedDir(), TEXT("Win64"), TEXT("Manifest"), TEXT("MicrosoftGame.config") );
diff --git a/Engine/Source/Runtime/Microsoft/GDKRuntime/Public/GDKRuntimeModule.h b/Engine/Source/Runtime/Microsoft/GDKRuntime/Public/GDKRuntimeModule.h
--- a/Engine/Source/Runtime/Microsoft/GDKRuntime/Public/GDKRuntimeModule.h
+++ b/Engine/Source/Runtime/Microsoft/GDKRuntime/Public/GDKRuntimeModule.h
@@ -19,10 +19,14 @@
 
 #include "Microsoft/AllowMicrosoftPlatformTypes.h"
 THIRD_PARTY_INCLUDES_START
+#include <grdk.h>
 #include <XAsync.h>
 THIRD_PARTY_INCLUDES_END
 #include "Microsoft/HideMicrosoftPlatformTypes.h"
 
+#ifndef WITH_GRDK_DEV_RUNTIME_INIT
+	#define WITH_GRDK_DEV_RUNTIME_INIT PLATFORM_WINDOWS && !UE_BUILD_SHIPPING && (_GRDK_EDITION >= 240600)
+#endif
 
 enum class XUserGamertagComponent: uint32_t;
 enum class XNetworkingConnectivityLevelHint : uint32_t;
